[NDIMedia] Fix ndi senders leaking in case of an error during the capture's initialization.

Problem:
In case of an error during the initialization of the capture, the state is changed to "Stopped", but StopCaptureImpl() is not called, even if calling StopCapture().
For NDI, this occurs after the sender is created, and it needs to be cleanly removed otherwise it will leak and cause name collision that will prevent restarting the capture.

Solution:
In MediaCapture, instead of setting the state to Stopped, properly call StopCapture.
However this lead to an infinite loop on bOutputResourcesInitialized (which was false because InitializeOutputResources returned early on error). The logic is better if inverted, i.e. waiting on bPendingOutputResourcesInitialized, only if it is true.

Also, for NDIMediaCapture, added a more explicit name collision detection given that it is a typical cause of failure. Hopefully this will help diagnose problems faster than just "an error has occured." message.

diff --git a/Engine/Plugins/Media/MediaIOFramework/Source/MediaIOCore/Private/MediaCapture.cpp b/Engine/Plugins/Media/MediaIOFramework/Source/MediaIOCore/Private/MediaCapture.cpp
--- a/Engine/Plugins/Media/MediaIOFramework/Source/MediaIOCore/Private/MediaCapture.cpp
+++ b/Engine/Plugins/Media/MediaIOFramework/Source/MediaIOCore/Private/MediaCapture.cpp
@@ -134,7 +134,7 @@ void FMediaCaptureOptions::PostSerialize(const FArchive& Ar)
 
 UMediaCapture::UMediaCapture()
 	: ValidSourceGPUMask(FRHIGPUMask::All())
-	, bOutputResourcesInitialized(false)
+	, bPendingOutputResourcesInitialized(false)
 	, bShouldCaptureRHIResource(false)
 	, WaitingForRenderCommandExecutionCounter(0)
 	, PendingFrameCount(0)
@@ -298,7 +298,7 @@ bool UMediaCapture::StartSourceCapture(TSharedPtr<UE::MediaCapture::Private::FCa
 	if (bInitialized)
 	{
 		InitializeOutputResources(MediaOutput->NumberOfTextureBuffers);
-		bInitialized = GetState() != EMediaCaptureState::Stopped;
+		bInitialized = GetState() != EMediaCaptureState::Stopped && GetState() != EMediaCaptureState::Error;
 	}
 
 	if (bInitialized)
@@ -330,7 +330,7 @@ bool UMediaCapture::StartSourceCapture(TSharedPtr<UE::MediaCapture::Private::FCa
 	}
 	else
 	{
-		SetState(EMediaCaptureState::Stopped);
+		StopCapture(/*bAllowPendingFrameToBeProcess*/ false);
 		MediaCaptureDetails::ShowSlateNotification();
 	}
 
@@ -562,7 +562,7 @@ void UMediaCapture::StopCapture(bool bAllowPendingFrameToBeProcess)
 			FCoreDelegates::OnBeginFrame.RemoveAll(this);
 			FCoreDelegates::OnEndFrame.RemoveAll(this);
 
-			while (WaitingForRenderCommandExecutionCounter.load() > 0 || !bOutputResourcesInitialized)
+			while (WaitingForRenderCommandExecutionCounter.load() > 0 || bPendingOutputResourcesInitialized)
 			{
 				FlushRenderingCommands();
 			}
@@ -679,7 +679,7 @@ bool UMediaCapture::CaptureImmediate_RenderThread(const UE::MediaCaptureData::FC
 	}
 
 	// This could happen if the capture immediate is called before our resources were initialized. We can't really flush as we are in a command.
-	if (!bOutputResourcesInitialized)
+	if (bPendingOutputResourcesInitialized)
 	{
 		UE_LOG(LogMediaIOCore, VeryVerbose, TEXT("Could not capture frame. Output resources haven't been initialized"));
 		return false;
@@ -913,7 +913,6 @@ void UMediaCapture::RestartCapture()
 
 		if (bInitialized)
 		{
-			bOutputResourcesInitialized = false;
 			InitializeOutputResources(MediaOutput->NumberOfTextureBuffers);
 			bInitialized = GetState() != EMediaCaptureState::Stopped;
 		}
@@ -1142,14 +1141,14 @@ void UMediaCapture::InitializeOutputResources(int32 InNumberOfBuffers)
 	if (DesiredOutputSize.X <= 0 || DesiredOutputSize.Y <= 0)
 	{
 		UE_LOG(LogMediaIOCore, Error, TEXT("Can't start the capture. The size requested is negative or zero."));
-		SetState(EMediaCaptureState::Stopped);
+		SetState(EMediaCaptureState::Error);
 		return;
 	}
 
 	// Recreate frame manager which can trigger cleaning up its captured frames if it exists
 	FrameManager = MakePimpl<FFrameManager>();
 
-	bOutputResourcesInitialized = false;
+	bPendingOutputResourcesInitialized = true;
 	bSyncHandlersInitialized = false;
 	NumberOfCaptureFrame = InNumberOfBuffers;
 
@@ -1204,7 +1203,7 @@ void UMediaCapture::InitializeOutputResources(int32 InNumberOfBuffers)
 					}
 				}
 			}
-			This->bOutputResourcesInitialized = true;
+			This->bPendingOutputResourcesInitialized = false;
 		});
 }
 
@@ -1329,7 +1328,7 @@ void UMediaCapture::OnEndFrame_GameThread()
 	TRACE_CPUPROFILER_EVENT_SCOPE(MediaCaptureEndFrame);
 	TRACE_CPUPROFILER_EVENT_SCOPE_TEXT(*FString::Printf(TEXT("MediaCapturePipe: %llu"), GFrameCounter % 10));
 
-	if (!bOutputResourcesInitialized)
+	if (bPendingOutputResourcesInitialized)
 	{
 		FlushRenderingCommands();
 	}
diff --git a/Engine/Plugins/Media/MediaIOFramework/Source/MediaIOCore/Public/MediaCapture.h b/Engine/Plugins/Media/MediaIOFramework/Source/MediaIOCore/Public/MediaCapture.h
--- a/Engine/Plugins/Media/MediaIOFramework/Source/MediaIOCore/Public/MediaCapture.h
+++ b/Engine/Plugins/Media/MediaIOFramework/Source/MediaIOCore/Public/MediaCapture.h
@@ -711,7 +711,7 @@ private:
 	bool bLoggedOutputNotReady = false;
 
 	bool bViewportHasFixedViewportSize = false;
-	std::atomic<bool> bOutputResourcesInitialized;
+	std::atomic<bool> bPendingOutputResourcesInitialized;
 	std::atomic<bool> bSyncHandlersInitialized;
 	std::atomic<bool> bShouldCaptureRHIResource;
 	std::atomic<int32> WaitingForRenderCommandExecutionCounter;
diff --git a/Engine/Plugins/Media/NDIMedia/Source/NDIMedia/Private/NDIMediaCapture.cpp b/Engine/Plugins/Media/NDIMedia/Source/NDIMedia/Private/NDIMediaCapture.cpp
--- a/Engine/Plugins/Media/NDIMedia/Source/NDIMedia/Private/NDIMediaCapture.cpp
+++ b/Engine/Plugins/Media/NDIMedia/Source/NDIMedia/Private/NDIMediaCapture.cpp
@@ -14,16 +14,19 @@ DEFINE_LOG_CATEGORY(LogNDIMedia);
 class UNDIMediaCapture::FNDICaptureInstance
 {
 public:
-	FNDICaptureInstance(const TSharedPtr<FNDIMediaRuntimeLibrary>& InNDILib, const UNDIMediaOutput* InMediaOutput)
+	FNDICaptureInstance(const TSharedPtr<FNDIMediaRuntimeLibrary>& InNDILib, const UNDIMediaOutput& InMediaOutput)
 		: NDILibHandle(InNDILib)
 		, NDILib(InNDILib ? InNDILib->Lib : nullptr)
 	{
+		FullSenderName = InMediaOutput.GroupName.IsEmpty() ? InMediaOutput.SourceName
+			: FString::Printf(TEXT("%s_%s"), *InMediaOutput.GroupName, *InMediaOutput.SourceName);
+		
 		if (NDILib != nullptr)
 		{
 			NDIlib_send_create_t SendDesc;
 
-			SendDesc.p_ndi_name = TCHAR_TO_UTF8(*InMediaOutput->SourceName);
-			SendDesc.p_groups = (!InMediaOutput->GroupName.IsEmpty()) ? TCHAR_TO_UTF8(*InMediaOutput->GroupName) : nullptr;
+			SendDesc.p_ndi_name = TCHAR_TO_UTF8(*InMediaOutput.SourceName);
+			SendDesc.p_groups = (!InMediaOutput.GroupName.IsEmpty()) ? TCHAR_TO_UTF8(*InMediaOutput.GroupName) : nullptr;
 
 			// Don't clock audio, normally, if audio and video is
 			SendDesc.clock_audio = false;
@@ -32,17 +35,29 @@ public:
 			SendDesc.clock_video = true;
 
 			Sender = NDILib->send_create(&SendDesc);
-		}
 
-		if (!Sender)
-		{
-			UE_LOG(LogNDIMedia, Error, TEXT("Failed to create NDI capture."));
+			if (Sender)
+			{
+				NDILibHandle->Senders.Add(FullSenderName);
+			}
+			else
+			{
+				// Likely cause of failure is having source name collision.
+				if (NDILibHandle->Senders.Contains(FullSenderName))
+				{
+					UE_LOG(LogNDIMedia, Error, TEXT("Failed to create NDI capture \"%s\". A source of the same name has already been created."), *FullSenderName);
+				}
+				else
+				{
+					UE_LOG(LogNDIMedia, Error, TEXT("Failed to create NDI capture \"%s\"."), *FullSenderName);	
+				}
+			}
 		}
 
 		// Keep track of specified frame rate.
-		FrameRateDenominator = InMediaOutput->FrameRate.Denominator;
-		FrameRateNumerator = InMediaOutput->FrameRate.Numerator;
-		OutputType = InMediaOutput->OutputType;
+		FrameRateDenominator = InMediaOutput.FrameRate.Denominator;
+		FrameRateNumerator = InMediaOutput.FrameRate.Numerator;
+		OutputType = InMediaOutput.OutputType;
 
 		// Caution: logic inversion, on purpose, because for this class, async 
 		// enables more work, while sync disables, and I prefer having my inverted
@@ -50,7 +65,7 @@ public:
 		// bWaitForSyncEvent logic in Media Output is inverted to match with BlackMedia
 		// and AJA Media Output's properties, in the hope that it makes it easier to
 		// generically manage those objects.
-		bAsyncSend = !InMediaOutput->bWaitForSyncEvent;
+		bAsyncSend = !InMediaOutput.bWaitForSyncEvent;
 
 		if (bAsyncSend)
 		{
@@ -70,11 +85,19 @@ public:
 	{
 		if (Sender)
 		{
-			// Force sync in case some data is still used by the ndi encoder.
-			NDILib->send_send_video_v2(Sender, nullptr);
+			if (NDILibHandle)
+			{
+				NDILibHandle->Senders.Remove(FullSenderName);
+			}
+			
+			if (NDILib)
+			{
+				// Force sync in case some data is still used by the ndi encoder.
+				NDILib->send_send_video_v2(Sender, nullptr);
 
-			// Destroy the NDI sender
-			NDILib->send_destroy(Sender);
+				// Destroy the NDI sender
+				NDILib->send_destroy(Sender);
+			}
 
 			Sender = nullptr;
 		}
@@ -133,6 +156,9 @@ public:
 public:
 	TSharedPtr<FNDIMediaRuntimeLibrary> NDILibHandle;
 	const NDIlib_v5* NDILib = nullptr;
+
+	/** Keep track of full sender name: "groupname_sourcename" for error handling purposes. */
+	FString FullSenderName;
 	
 	NDIlib_send_instance_t Sender = nullptr;
 	int32 FrameRateNumerator = 30000;
@@ -291,7 +317,14 @@ bool UNDIMediaCapture::StartNewCapture()
 
 		if (const UNDIMediaOutput* NDIMediaOutput = Cast<UNDIMediaOutput>(MediaOutput))
 		{
-			CaptureInstance = new FNDICaptureInstance(FNDIMediaModule::GetNDIRuntimeLibrary(), NDIMediaOutput);
+			CaptureInstance = new FNDICaptureInstance(FNDIMediaModule::GetNDIRuntimeLibrary(), *NDIMediaOutput);
+
+			// Validate that the sender has been created.
+			if (CaptureInstance->Sender)
+			{
+				SetState(EMediaCaptureState::Capturing);
+				return true;
+			}
 		}
 		else
 		{
@@ -299,6 +332,5 @@ bool UNDIMediaCapture::StartNewCapture()
 		}
 	}
 
-	SetState(EMediaCaptureState::Capturing);
-	return true;
+	return false;
 }
diff --git a/Engine/Plugins/Media/NDIMedia/Source/NDIMedia/Private/NDIMediaModule.h b/Engine/Plugins/Media/NDIMedia/Source/NDIMedia/Private/NDIMediaModule.h
--- a/Engine/Plugins/Media/NDIMedia/Source/NDIMedia/Private/NDIMediaModule.h
+++ b/Engine/Plugins/Media/NDIMedia/Source/NDIMedia/Private/NDIMediaModule.h
@@ -25,6 +25,9 @@ public:
 	void* LibHandle = nullptr;
 
 	FString LibraryPath;
+
+	/** Keep track of senders being created to detect source collisions and provide better error messages. */
+	TSet<FString> Senders;	// Format: GroupName_SourceName
 };
 
 class FNDIMediaModule : public IModuleInterface
diff --git a/Engine/Plugins/Media/NDIMedia/Source/NDIMedia/Public/NDIMediaCapture.h b/Engine/Plugins/Media/NDIMedia/Source/NDIMedia/Public/NDIMediaCapture.h
--- a/Engine/Plugins/Media/NDIMedia/Source/NDIMedia/Public/NDIMediaCapture.h
+++ b/Engine/Plugins/Media/NDIMedia/Source/NDIMedia/Public/NDIMediaCapture.h
@@ -4,8 +4,6 @@
 
 #include "MediaCapture.h"
 #include "HAL/CriticalSection.h"
-#include "Misc/FrameRate.h"
-#include "NDIMediaOutput.h"
 
 #include "NDIMediaCapture.generated.h"
 
@@ -20,7 +18,7 @@ class NDIMEDIA_API UNDIMediaCapture : public UMediaCapture
 	GENERATED_BODY()
 
 public:
-	virtual ~UNDIMediaCapture();	// so we can pimpl.
+	virtual ~UNDIMediaCapture();
 
 protected:
 	virtual void OnFrameCaptured_RenderingThread(const FCaptureBaseData& InBaseData, TSharedPtr<FMediaCaptureUserData, ESPMode::ThreadSafe> InUserData, void* InBuffer, int32 Width, int32 Height, int32 BytesPerRow) override;
