[NDIMedia] Implementing Audio Ouput.

- Includes a change to MediaIOFramework FMediaIOAudioOutput to support float buffer upmixing.

diff --git a/Engine/Plugins/Media/MediaIOFramework/Source/MediaIOCore/Public/MediaIOCoreAudioOutput.h b/Engine/Plugins/Media/MediaIOFramework/Source/MediaIOCore/Public/MediaIOCoreAudioOutput.h
--- a/Engine/Plugins/Media/MediaIOFramework/Source/MediaIOCore/Public/MediaIOCoreAudioOutput.h
+++ b/Engine/Plugins/Media/MediaIOFramework/Source/MediaIOCore/Public/MediaIOCoreAudioOutput.h
@@ -1,4 +1,4 @@
-﻿// Copyright Epic Games, Inc. All Rights Reserved.
+// Copyright Epic Games, Inc. All Rights Reserved.
 
 #pragma once
 
@@ -20,7 +20,7 @@ class FAudioDeviceHandle;
 DECLARE_DELEGATE_TwoParams(FOnBufferReceived, uint8* Buffer, int32 BufferSize);
 
 class MEDIAIOCORE_API FMediaIOAudioOutput
-{ 
+{
 public:
 	struct FAudioOptions
 	{
@@ -53,20 +53,14 @@ public:
 			ConvertedBuffer.SetNumZeroed(FMath::CeilToInt32(ChannelRatio * InBuffer.Num()));
 
 			OutputType* ConvertedBufferPtr = ConvertedBuffer.GetData();
+			const TSampleConverter<OutputType> SamplerConverter;
 
 			for (int32 Index = 0; Index < InBuffer.Num(); Index += InNumInputChannels)
 			{
-				constexpr double Scale = TNumericLimits<OutputType>::Max();
-
-				// @Note: While this conversion is a common approach, it can introduce distortion. 
-				// See: https://www.cs.cmu.edu/~rbd/papers/cmj-float-to-int.html
 				// Convert values for each channel
 				for (int32 ChannelIndex = 0; ChannelIndex < InNumInputChannels; ++ChannelIndex)
 				{
-					const float FloatValue = InBuffer[Index + ChannelIndex];
-					OutputType ConvertedValue = static_cast<OutputType>((FloatValue * Scale) + 0.5);
-
-					*(ConvertedBufferPtr + ChannelIndex) = ConvertedValue;
+					*(ConvertedBufferPtr + ChannelIndex) = SamplerConverter.ConvertSample(InBuffer[Index + ChannelIndex]);
 				}
 
 				ConvertedBufferPtr += InNumOutputChannels;
@@ -141,6 +135,24 @@ private:
 	int32 GetAudioBuffer(int32 InNumSamplesToPop, float* OutBuffer) const;
 	Audio::FAlignedFloatBuffer GetFloatBuffer(uint32 NumSamplesToGet) const;
 
+	/** Sample conversion helper template. */
+	template <typename OutputType>
+	struct TSampleConverter
+	{
+		const double Scale;
+
+		constexpr TSampleConverter()
+			: Scale(TNumericLimits<OutputType>::Max())
+		{}
+		
+		OutputType ConvertSample(const float InFloatValue) const
+		{
+			// @Note: While this conversion is a common approach, it can introduce distortion. 
+			// See: https://www.cs.cmu.edu/~rbd/papers/cmj-float-to-int.html
+			return static_cast<OutputType>((InFloatValue * Scale) + 0.5);
+		}
+	};
+
 private:
 	/** The buffer accumulating audio samples. */
 	Audio::FPatchOutputStrongPtr PatchOutput;
@@ -149,6 +161,16 @@ private:
 	FOnBufferReceived BufferReceivedDelegate;
 };
 
+/** Specialization for float buffers. */
+template <>
+struct FMediaIOAudioOutput::TSampleConverter<float>
+{
+	float ConvertSample(const float InFloatValue) const
+	{
+		return InFloatValue; // no op for float.
+	}
+};
+
 /**
  * Handles capturing audio samples rendered by the engine and dispatching them to outputs.
  */
diff --git a/Engine/Plugins/Media/NDIMedia/Source/NDIMedia/Private/NDIMediaCapture.cpp b/Engine/Plugins/Media/NDIMedia/Source/NDIMedia/Private/NDIMediaCapture.cpp
--- a/Engine/Plugins/Media/NDIMedia/Source/NDIMedia/Private/NDIMediaCapture.cpp
+++ b/Engine/Plugins/Media/NDIMedia/Source/NDIMedia/Private/NDIMediaCapture.cpp
@@ -2,11 +2,12 @@
 
 #include "NDIMediaCapture.h"
 
+#include "Engine/Engine.h"
+#include "Engine/TextureRenderTarget2D.h"
+#include "MediaIOCoreSubsystem.h"
 #include "NDIMediaAPI.h"
 #include "NDIMediaModule.h"
 #include "NDIMediaOutput.h"
-
-#include "Engine/TextureRenderTarget2D.h"
 #include "Slate/SceneViewport.h"
 
 DEFINE_LOG_CATEGORY(LogNDIMedia);
@@ -55,8 +56,7 @@ public:
 		}
 
 		// Keep track of specified frame rate.
-		FrameRateDenominator = InMediaOutput.FrameRate.Denominator;
-		FrameRateNumerator = InMediaOutput.FrameRate.Numerator;
+		FrameRate = InMediaOutput.FrameRate;
 		OutputType = InMediaOutput.OutputType;
 
 		// Caution: logic inversion, on purpose, because for this class, async 
@@ -79,6 +79,10 @@ public:
 
 			VideoFrameBuffers.SetNum(NumVideoFrameBuffers);
 		}
+
+		// Prepare the audio frame circular buffer.
+		static constexpr int32 NumAudioFrameBuffers = 2;
+		AudioFrameBuffers.SetNum(NumAudioFrameBuffers);
 	}
 
 	~FNDICaptureInstance()
@@ -152,6 +156,100 @@ public:
 		// Ensure video frame buffer is of proper size.
 		return VideoFrameBuffers[VideoFrameBufferCurrentIndex]->EnsureSize(InHeight, InBytesPerRow);
 	}
+	
+	bool UpdateAudioOutput(const FAudioDeviceHandle& InAudioDeviceHandle, const UNDIMediaOutput& InMediaOutput)
+	{
+		bSendAudioOnlyIfReceiversConnected = InMediaOutput.bSendAudioOnlyIfReceiversConnected;
+
+		if (GEngine && InMediaOutput.bOutputAudio)
+		{
+			UMediaIOCoreSubsystem::FCreateAudioOutputArgs Args;
+			Args.NumOutputChannels = static_cast<int32>(InMediaOutput.NumOutputAudioChannels);
+			Args.TargetFrameRate = InMediaOutput.FrameRate;
+			Args.MaxSampleLatency = Align(InMediaOutput.AudioBufferSize, 4);
+			Args.OutputSampleRate = static_cast<uint32>(InMediaOutput.AudioSampleRate);
+			Args.AudioDeviceHandle = InAudioDeviceHandle;
+			AudioOutput = GEngine->GetEngineSubsystem<UMediaIOCoreSubsystem>()->CreateAudioOutput(Args);
+			return AudioOutput.IsValid();
+		}
+		
+		AudioOutput.Reset();
+		return false;
+	}
+
+	struct FAudioFrameBuffer
+	{
+		TArray<float> Data;
+
+		FAudioFrameBuffer(int32 InNumSamples)
+		{
+			Data.Reset(InNumSamples);
+		}
+
+		float* Reset(int32 InNumSamples)
+		{
+			Data.Reset(InNumSamples);
+			return Data.GetData();
+		}
+	};
+	
+	float* GetNextAudioFrameBuffer(int32 InNumSamples)
+	{
+		// Move to next audio frame buffer in the circular array.
+		AudioFrameBufferCurrentIndex++;
+		if (AudioFrameBufferCurrentIndex >= AudioFrameBuffers.Num())
+		{
+			AudioFrameBufferCurrentIndex = 0;
+		}
+
+		// Lazy allocation
+		if (!AudioFrameBuffers[AudioFrameBufferCurrentIndex].IsValid())
+		{
+			AudioFrameBuffers[AudioFrameBufferCurrentIndex] = MakeUnique<FAudioFrameBuffer>(InNumSamples);
+		}
+
+		// Ensure audio frame buffer is of proper size.
+		return AudioFrameBuffers[AudioFrameBufferCurrentIndex]->Reset(InNumSamples);
+	}
+	
+	void OutputAudio(int64 InTimeCode)
+	{
+		TRACE_CPUPROFILER_EVENT_SCOPE(UNDIMediaCapture::OutputAudio);
+
+		// Take a local copy of the audio output in case it is switched from the main thread.
+		const TSharedPtr<FMediaIOAudioOutput> LocalAudioOutput = AudioOutput;
+		if(!LocalAudioOutput)
+		{
+			return;
+		}
+
+		// This returns an interleaved buffer with NumOutputChannels.
+		TArray<float> InterleavedAudioBuffer = LocalAudioOutput->GetAllAudioSamples<float>();
+		
+		if (!Sender || (bSendAudioOnlyIfReceiversConnected && NDILib->send_get_no_connections(Sender, 0) <= 0))
+		{
+			return;
+		}
+		
+		// Convert from the interleaved audio that Unreal Engine produces
+		const int32 NumChannels = LocalAudioOutput->NumOutputChannels;
+		const int32 NumSamples = InterleavedAudioBuffer.Num();
+		const int32 NumSamplesPerChannel = InterleavedAudioBuffer.Num() / NumChannels;
+
+		NDIlib_audio_frame_interleaved_32f_t NDI_interleaved_audio_frame;
+		NDI_interleaved_audio_frame.timecode = InTimeCode;
+		NDI_interleaved_audio_frame.sample_rate = LocalAudioOutput->OutputSampleRate;
+		NDI_interleaved_audio_frame.no_channels = NumChannels;	
+		NDI_interleaved_audio_frame.no_samples = NumSamplesPerChannel;
+		NDI_interleaved_audio_frame.p_data = InterleavedAudioBuffer.GetData();
+	
+		NDIlib_audio_frame_v2_t NDI_audio_frame;
+		NDI_audio_frame.p_data = GetNextAudioFrameBuffer(NumSamples);
+		NDI_audio_frame.channel_stride_in_bytes = (NumSamplesPerChannel) * sizeof(float);
+
+		NDILib->util_audio_from_interleaved_32f_v2(&NDI_interleaved_audio_frame, &NDI_audio_frame);
+		NDILib->send_send_audio_v2(Sender, &NDI_audio_frame);
+	}
 
 public:
 	TSharedPtr<FNDIMediaRuntimeLibrary> NDILibHandle;
@@ -161,16 +259,24 @@ public:
 	FString FullSenderName;
 	
 	NDIlib_send_instance_t Sender = nullptr;
-	int32 FrameRateNumerator = 30000;
-	int32 FrameRateDenominator = 1001;
+	FFrameRate FrameRate;
 	EMediaIOOutputType OutputType = EMediaIOOutputType::Fill;
 
-	// By default send async because it is the recommended way in the SDK.
+	/** By default send async because it is the recommended way in the SDK. */
 	bool bAsyncSend = true;
 
-	// Circular buffer of Video Frames.
+	/** Circular buffer of Video Frames. */
 	TArray<TUniquePtr<FVideoFrameBuffer>> VideoFrameBuffers;
 	int32 VideoFrameBufferCurrentIndex = 0;
+
+	/** Circular buffer of Audio Frames. */
+	TArray<TUniquePtr<FAudioFrameBuffer>> AudioFrameBuffers;
+	int32 AudioFrameBufferCurrentIndex = 0;
+
+	/** Holds an audio output that will receive samples from the media io core subsystem. */
+	TSharedPtr<FMediaIOAudioOutput> AudioOutput;
+
+	bool bSendAudioOnlyIfReceiversConnected = true;
 };
 
 UNDIMediaCapture::~UNDIMediaCapture()
@@ -237,10 +343,12 @@ void UNDIMediaCapture::OnFrameCaptured_RenderingThread(const FCaptureBaseData& I
 		NDI_video_frame.FourCC = bIsRGBA ? NDIlib_FourCC_type_BGRA : NDIlib_FourCC_type_UYVY;
 		NDI_video_frame.p_data = static_cast<uint8_t*>(InBuffer);
 		NDI_video_frame.line_stride_in_bytes = BytesPerRow;
-		NDI_video_frame.frame_rate_D = CaptureInstance->FrameRateDenominator;
-		NDI_video_frame.frame_rate_N = CaptureInstance->FrameRateNumerator;
+		NDI_video_frame.frame_rate_D = CaptureInstance->FrameRate.Denominator;
+		NDI_video_frame.frame_rate_N = CaptureInstance->FrameRate.Numerator;
 		NDI_video_frame.timecode = ConvertToNDITimeCode(InBaseData.SourceFrameTimecode, InBaseData.SourceFrameTimecodeFramerate);
 
+		CaptureInstance->OutputAudio(NDI_video_frame.timecode);
+
 		if (CaptureInstance->bAsyncSend)
 		{
 			// For async send, the memory buffer needs to remain valid until the next call.
@@ -296,6 +404,19 @@ bool UNDIMediaCapture::PostInitializeCaptureRenderTarget(UTextureRenderTarget2D*
 	return bSuccess;
 }
 
+bool UNDIMediaCapture::UpdateAudioDeviceImpl(const FAudioDeviceHandle& InAudioDeviceHandle)
+{
+	FScopeLock ScopeLock(&CaptureInstanceCriticalSection);
+	if (CaptureInstance)
+	{
+		if (const UNDIMediaOutput* NdiMediaOutput = Cast<UNDIMediaOutput>(MediaOutput))
+		{
+			return CaptureInstance->UpdateAudioOutput(InAudioDeviceHandle, *NdiMediaOutput);
+		}
+	}
+	return false;
+}
+
 void UNDIMediaCapture::StopCaptureImpl(bool /*bAllowPendingFrameToBeProcess*/)
 {
 	TRACE_BOOKMARK(TEXT("NDIMediaCapture::StopCapture"));
@@ -322,6 +443,7 @@ bool UNDIMediaCapture::StartNewCapture()
 			// Validate that the sender has been created.
 			if (CaptureInstance->Sender)
 			{
+				CaptureInstance->UpdateAudioOutput(AudioDeviceHandle, *NDIMediaOutput);
 				SetState(EMediaCaptureState::Capturing);
 				return true;
 			}
diff --git a/Engine/Plugins/Media/NDIMedia/Source/NDIMedia/Public/NDIMediaCapture.h b/Engine/Plugins/Media/NDIMedia/Source/NDIMedia/Public/NDIMediaCapture.h
--- a/Engine/Plugins/Media/NDIMedia/Source/NDIMedia/Public/NDIMediaCapture.h
+++ b/Engine/Plugins/Media/NDIMedia/Source/NDIMedia/Public/NDIMediaCapture.h
@@ -25,6 +25,7 @@ protected:
 	virtual bool InitializeCapture() override;
 	virtual bool PostInitializeCaptureViewport(TSharedPtr<FSceneViewport>& InSceneViewport) override;
 	virtual bool PostInitializeCaptureRenderTarget(UTextureRenderTarget2D* InRenderTarget) override;
+	virtual bool UpdateAudioDeviceImpl(const FAudioDeviceHandle& InAudioDeviceHandle) override;
 	virtual void StopCaptureImpl(bool bAllowPendingFrameToBeProcess) override;
 
 private:
diff --git a/Engine/Plugins/Media/NDIMedia/Source/NDIMedia/Public/NDIMediaOutput.h b/Engine/Plugins/Media/NDIMedia/Source/NDIMedia/Public/NDIMediaOutput.h
--- a/Engine/Plugins/Media/NDIMedia/Source/NDIMedia/Public/NDIMediaOutput.h
+++ b/Engine/Plugins/Media/NDIMedia/Source/NDIMedia/Public/NDIMediaOutput.h
@@ -17,6 +17,21 @@ enum class ENDIMediaOutputPixelFormat
 	B8G8R8A8					UMETA(DisplayName = "8bit RGBA"),
 };
 
+/**
+ * NDI Audio supports any sample rate.
+ * The standard for live production is 48 kHz.
+ */
+UENUM()
+enum class ENDIMediaOutputAudioSampleRate : uint32
+{
+	SR_44_1k = 44100 UMETA(DisplayName = "44.1 kHz"),
+	SR_48k = 48000 UMETA(DisplayName = "48 kHz"),
+	SR_88_2k = 88200 UMETA(DisplayName = "88.2 kHz"),
+	SR_96k = 96000 UMETA(DisplayName = "96 kHz"),
+	SR_176_4k = 176400 UMETA(DisplayName = "176.4 kHz"),
+	SR_192k = 192000 UMETA(DisplayName = "192 kHz")
+};
+
 // Note: Other formats YUV 8 bits or YUV 16 bits, both support fill and key as well.
 //
 // Format mapping:
@@ -83,7 +98,38 @@ public:
 	/** Represents the desired number of frames (per second) for video to be sent over NDI */
 	UPROPERTY(BlueprintReadwrite, EditDefaultsOnly, Category = "Media")
 	FFrameRate FrameRate = FFrameRate(60, 1);
+	
+	/** Whether to capture and output audio from the engine. */
+	UPROPERTY(BlueprintReadWrite, EditAnywhere, Category = "Audio")
+	bool bOutputAudio = false;
+
+	/** Size of the buffer that holds rendered audio samples, a bigger buffer will produce a more stable output signal but will introduce more delay. */
+	UPROPERTY(BlueprintReadWrite, EditAnywhere, Category = "Audio")
+	int32 AudioBufferSize = 5*1024;
+	
+	/**
+	 * An individual NDI stream can transport multiple audio channels.
+	 * The number of channels supported depends on the codec used.
+	 * PCM allows unlimited channels; in NDI, AAC can support 2 channels,
+	 * while Opus can support up to 255 channels.
+	 */
+	UPROPERTY(BlueprintReadWrite, EditAnywhere, Category = "Audio")
+	int32 NumOutputAudioChannels = 2;
+
+	/**
+	 * NDI Audio supports any sample rate.
+	 * The standard for live production is 48 kHz.
+	 */
+	UPROPERTY(BlueprintReadWrite, EditAnywhere, Category = "Audio")
+	ENDIMediaOutputAudioSampleRate AudioSampleRate = ENDIMediaOutputAudioSampleRate::SR_48k;
 
+	/**
+	 * As an optimization, the audio will not be converted and sent if there are no connected receivers.
+	 * Setting this to false will result in audio being converted and sent regardless of receivers.
+	 */
+	UPROPERTY(BlueprintReadWrite, EditAnywhere, Category = "Audio")
+	bool bSendAudioOnlyIfReceiversConnected = true;
+	
 	/** 
 	 * Wait for the NDI source sync event.
 	 * Caution: this will be blocking in the rendering thread and may cause performance issues in the engine.
