[NDIMedia] NDI Media Input implementations (first pass)

- Media Source & Player
- Timecode provider

diff --git a/Engine/Plugins/Media/NDIMedia/Source/NDIMedia/NDIMedia.Build.cs b/Engine/Plugins/Media/NDIMedia/Source/NDIMedia/NDIMedia.Build.cs
--- a/Engine/Plugins/Media/NDIMedia/Source/NDIMedia/NDIMedia.Build.cs
+++ b/Engine/Plugins/Media/NDIMedia/Source/NDIMedia/NDIMedia.Build.cs
@@ -11,7 +11,8 @@ public class NDIMedia : ModuleRules
 		PublicDependencyModuleNames.AddRange(
 			new string[] {
 				"Core",
-				"MediaIOCore",
+				"MediaAssets",
+				"MediaIOCore"
 			});
 
 		PrivateDependencyModuleNames.AddRange(
@@ -19,10 +20,15 @@ public class NDIMedia : ModuleRules
 				"CoreUObject",
 				"DeveloperSettings",
 				"Engine",
+				"MediaIOCore",
+				"MediaUtils",
 				"NDISDK",
 				"Projects",
+				"RenderCore",
+				"RHI",
 				"Slate",
 				"SlateCore",
+				"TimeManagement",
 			});
 
 		if (Target.bBuildEditor == true)
diff --git a/Engine/Plugins/Media/NDIMedia/Source/NDIMedia/Private/NDIDeviceProvider.cpp b/Engine/Plugins/Media/NDIMedia/Source/NDIMedia/Private/NDIDeviceProvider.cpp
new file mode 100644
--- /dev/null
+++ b/Engine/Plugins/Media/NDIMedia/Source/NDIMedia/Private/NDIDeviceProvider.cpp
@@ -0,0 +1,251 @@
+// Copyright Epic Games, Inc. All Rights Reserved.
+
+#include "NDIDeviceProvider.h"
+
+#include "Internationalization/Text.h"
+#include "MediaIOCoreDefinitions.h"
+#include "NDIMediaModule.h"
+#include "NDISourceFinder.h"
+
+#define LOCTEXT_NAMESPACE "NDIDeviceProvider"
+
+namespace UE::NDIDeviceProvider::Private
+{
+	TSharedPtr<FNDISourceFinder> GetFindInstance()
+	{
+		if (FNDIMediaModule* NdiModule = FNDIMediaModule::Get())
+		{
+			return NdiModule->GetFindInstance();
+		}
+		return nullptr;
+	}
+}
+
+FName FNDIDeviceProvider::GetProviderName()
+{
+	static FName NAME_Provider = "NDI";
+	return NAME_Provider;
+}
+
+FName FNDIDeviceProvider::GetProtocolName()
+{
+	static FName NAME_Protocol = "ndi";
+	return NAME_Protocol;
+}
+
+FName FNDIDeviceProvider::GetFName()
+{
+	return GetProviderName();
+}
+
+TArray<FMediaIOConnection> FNDIDeviceProvider::GetConnections() const
+{
+	return TArray<FMediaIOConnection>();
+}
+
+TArray<FMediaIOConfiguration> FNDIDeviceProvider::GetConfigurations() const
+{
+	return GetConfigurations(true, true);
+}
+
+TArray<FMediaIOConfiguration> FNDIDeviceProvider::GetConfigurations(bool bInAllowInput, bool /*bInAllowOutput*/) const
+{
+	TArray<FMediaIOConfiguration> Results;
+	if (bInAllowInput)
+	{
+		if (const TSharedPtr<FNDISourceFinder> FindInstance = UE::NDIDeviceProvider::Private::GetFindInstance())
+		{
+			TArray<FNDISourceFinder::FNDISourceInfo> Sources = FindInstance->GetSources();
+
+			int32 DeviceId = 0;
+			for (const FNDISourceFinder::FNDISourceInfo& Source : Sources)
+			{
+				FMediaIOConfiguration MediaConfiguration = GetDefaultConfiguration();
+				MediaConfiguration.bIsInput = true;
+				// Remark: we would also like to keep the URL, but it can be recovered again, so source name is sufficient for now.
+				MediaConfiguration.MediaConnection.Device.DeviceName = *Source.Name;
+				MediaConfiguration.MediaConnection.Device.DeviceIdentifier = DeviceId;
+				MediaConfiguration.MediaConnection.Protocol = GetProtocolName();
+				MediaConfiguration.MediaConnection.PortIdentifier = 0;
+
+				++DeviceId;
+
+				Results.Add(MediaConfiguration);
+			}
+		}
+	}
+	return Results;
+}
+
+TArray<FMediaIOInputConfiguration> FNDIDeviceProvider::GetInputConfigurations() const
+{
+	TArray<FMediaIOInputConfiguration> Results;
+
+	TArray<FMediaIOConfiguration> Configs = GetConfigurations(/*bInAllowInput*/true, /*bInAllowOutput*/false);
+
+	FMediaIOInputConfiguration DefaultInputConfiguration = GetDefaultInputConfiguration();
+	DefaultInputConfiguration.KeyPortIdentifier = 0;
+	DefaultInputConfiguration.InputType = EMediaIOInputType::FillAndKey; // NDI supports alpha channel
+	
+	for (const FMediaIOConfiguration& Config : Configs)
+	{
+		DefaultInputConfiguration.MediaConfiguration = Config;
+		Results.Add(DefaultInputConfiguration);
+	}
+	return Results;
+
+}
+
+TArray<FMediaIOOutputConfiguration> FNDIDeviceProvider::GetOutputConfigurations() const
+{
+	return TArray<FMediaIOOutputConfiguration>();
+}
+
+TArray<FMediaIOVideoTimecodeConfiguration> FNDIDeviceProvider::GetTimecodeConfigurations() const
+{
+	TArray<FMediaIOVideoTimecodeConfiguration> MediaConfigurations;
+	{
+		TArray<FMediaIOConfiguration> InputConfigurations = GetConfigurations(true, false);
+
+		FMediaIOVideoTimecodeConfiguration DefaultTimecodeConfiguration;
+		MediaConfigurations.Reset(InputConfigurations.Num());
+		for (const FMediaIOConfiguration& InputConfiguration : InputConfigurations)
+		{
+			DefaultTimecodeConfiguration.MediaConfiguration = InputConfiguration;
+			DefaultTimecodeConfiguration.TimecodeFormat = EMediaIOAutoDetectableTimecodeFormat::LTC;	// todo: verify this
+			MediaConfigurations.Add(DefaultTimecodeConfiguration);
+		}
+	}
+	return MediaConfigurations;
+}
+
+TArray<FMediaIODevice> FNDIDeviceProvider::GetDevices() const
+{
+	TArray<FMediaIODevice> Results;
+
+	if (const TSharedPtr<FNDISourceFinder> FindInstance = UE::NDIDeviceProvider::Private::GetFindInstance())
+	{
+		TArray<FNDISourceFinder::FNDISourceInfo> Sources = FindInstance->GetSources();
+
+		int32 DeviceId = 0;
+		for (const FNDISourceFinder::FNDISourceInfo& Source : Sources)
+		{
+			FMediaIODevice Device;
+			Device.DeviceName = *Source.Name;
+			Device.DeviceIdentifier = DeviceId;
+			Results.Add(Device);
+			++DeviceId;
+		}
+	}
+
+	return Results;
+}
+
+
+TArray<FMediaIOMode> FNDIDeviceProvider::GetModes(const FMediaIODevice& InDevice, bool bInOutput) const
+{
+	TArray<FMediaIOMode> Results;
+	return Results;
+}
+
+FMediaIOConfiguration FNDIDeviceProvider::GetDefaultConfiguration() const
+{
+	FMediaIOConfiguration Configuration;
+	Configuration.bIsInput = true;
+	Configuration.MediaConnection.Device.DeviceIdentifier = 1;
+	Configuration.MediaConnection.Protocol = GetProtocolName();
+	Configuration.MediaConnection.PortIdentifier = 0;
+	Configuration.MediaMode = GetDefaultMode();
+	return Configuration;
+}
+
+
+FMediaIOMode FNDIDeviceProvider::GetDefaultMode() const
+{
+	FMediaIOMode Mode;
+	Mode.DeviceModeIdentifier = 0;	// Unused, but can't be invalid.
+	Mode.FrameRate = FFrameRate(30, 1);
+	Mode.Resolution = FIntPoint(1920, 1080);
+	Mode.Standard = EMediaIOStandardType::Progressive;
+	return Mode;
+}
+
+
+FMediaIOInputConfiguration FNDIDeviceProvider::GetDefaultInputConfiguration() const
+{
+	FMediaIOInputConfiguration Configuration;
+	Configuration.MediaConfiguration = GetDefaultConfiguration();
+	Configuration.MediaConfiguration.bIsInput = true;
+	Configuration.InputType = EMediaIOInputType::FillAndKey;
+	return Configuration;
+}
+
+FMediaIOOutputConfiguration FNDIDeviceProvider::GetDefaultOutputConfiguration() const
+{
+	FMediaIOOutputConfiguration Configuration;
+	Configuration.MediaConfiguration = GetDefaultConfiguration();
+	Configuration.MediaConfiguration.bIsInput = false;
+	Configuration.OutputReference = EMediaIOReferenceType::FreeRun;
+	Configuration.OutputType = EMediaIOOutputType::FillAndKey;
+	return Configuration;
+}
+
+FMediaIOVideoTimecodeConfiguration FNDIDeviceProvider::GetDefaultTimecodeConfiguration() const
+{
+	FMediaIOVideoTimecodeConfiguration Configuration;
+	Configuration.MediaConfiguration = GetDefaultConfiguration();
+	return Configuration;
+}
+
+FText FNDIDeviceProvider::ToText(const FMediaIOConfiguration& InConfiguration, bool bInIsAutoDetected) const
+{
+	if (bInIsAutoDetected)
+	{
+		return FText::Format(LOCTEXT("FMediaIOAutoConfigurationToText", "{0} - {1} [device{2}/auto]")
+				, InConfiguration.bIsInput ? LOCTEXT("In", "In") : LOCTEXT("Out", "Out")
+				, FText::FromName(InConfiguration.MediaConnection.Device.DeviceName)
+				, FText::AsNumber(InConfiguration.MediaConnection.Device.DeviceIdentifier)
+				);
+	}
+	if (InConfiguration.IsValid())
+	{
+		return FText::Format(LOCTEXT("FMediaIOConfigurationToText", "[{0}] - {1} [device{2}/{3}]")
+			, InConfiguration.bIsInput ? LOCTEXT("In", "In") : LOCTEXT("Out", "Out")
+			, FText::FromName(InConfiguration.MediaConnection.Device.DeviceName)
+			, FText::AsNumber(InConfiguration.MediaConnection.Device.DeviceIdentifier)
+			, InConfiguration.MediaMode.GetModeName()
+		);
+	}
+	return LOCTEXT("Invalid", "<Invalid>");
+}
+
+
+FText FNDIDeviceProvider::ToText(const FMediaIOConnection& InConnection) const
+{
+	if (InConnection.IsValid())
+	{
+		return FText::Format(LOCTEXT("FMediaIOConnectionToText", "{0} [device{1}]")
+			, FText::FromName(InConnection.Device.DeviceName)
+			, LOCTEXT("Device", "device")
+			, FText::AsNumber(InConnection.Device.DeviceIdentifier)
+		);
+	}
+	return LOCTEXT("Invalid", "<Invalid>");
+}
+
+FText FNDIDeviceProvider::ToText(const FMediaIOOutputConfiguration& InConfiguration) const
+{
+	if (InConfiguration.IsValid())
+	{
+		return FText::Format(LOCTEXT("FMediaIOOutputConfigurationToText", "{0} - {1} [device{2}/{3}/{4}]")
+			, InConfiguration.OutputType == EMediaIOOutputType::Fill ? LOCTEXT("Fill", "Fill") : LOCTEXT("FillAndKey", "Fill&Key")
+			, FText::FromName(InConfiguration.MediaConfiguration.MediaConnection.Device.DeviceName)
+			, FText::AsNumber(InConfiguration.MediaConfiguration.MediaConnection.Device.DeviceIdentifier)
+			, GetTransportName(InConfiguration.MediaConfiguration.MediaConnection.TransportType, InConfiguration.MediaConfiguration.MediaConnection.QuadTransportType)
+			, InConfiguration.MediaConfiguration.MediaMode.GetModeName()
+		);
+	}
+	return LOCTEXT("Invalid", "<Invalid>");
+}
+
+#undef LOCTEXT_NAMESPACE
diff --git a/Engine/Plugins/Media/NDIMedia/Source/NDIMedia/Private/NDIDeviceProvider.h b/Engine/Plugins/Media/NDIMedia/Source/NDIMedia/Private/NDIDeviceProvider.h
new file mode 100644
--- /dev/null
+++ b/Engine/Plugins/Media/NDIMedia/Source/NDIMedia/Private/NDIDeviceProvider.h
@@ -0,0 +1,55 @@
+// Copyright Epic Games, Inc. All Rights Reserved.
+
+#pragma once
+
+#include "Containers/Array.h"
+#include "IMediaIOCoreDeviceProvider.h"
+
+class FText;
+struct FMediaIOConfiguration;
+struct FMediaIOConnection;
+struct FMediaIODevice;
+struct FMediaIOInputConfiguration;
+struct FMediaIOMode;
+struct FMediaIOOutputConfiguration;
+struct FMediaIOVideoTimecodeConfiguration;
+
+/**
+ * Implementation of IMediaIOCoreDeviceProvider for UNDIMediaSource "device" enumeration.
+ */
+class FNDIDeviceProvider : public IMediaIOCoreDeviceProvider
+{
+public:
+	static FName GetProviderName();
+	static FName GetProtocolName();
+
+	//~ Begin IMediaIOCoreDeviceProvider
+	virtual FName GetFName() override;
+
+	virtual TArray<FMediaIOConnection> GetConnections() const override;
+	virtual TArray<FMediaIOConfiguration> GetConfigurations() const override;
+	virtual TArray<FMediaIOConfiguration> GetConfigurations(bool bInAllowInput, bool bInAllowOutput) const override;
+	virtual TArray<FMediaIODevice> GetDevices() const override;
+	virtual TArray<FMediaIOMode> GetModes(const FMediaIODevice& InDevice, bool bInOutput) const override;
+	virtual TArray<FMediaIOInputConfiguration> GetInputConfigurations() const override;
+	virtual TArray<FMediaIOOutputConfiguration> GetOutputConfigurations() const override;
+	virtual TArray<FMediaIOVideoTimecodeConfiguration> GetTimecodeConfigurations() const override;
+	
+	virtual FMediaIOConfiguration GetDefaultConfiguration() const override;
+	virtual FMediaIOMode GetDefaultMode() const override;
+	virtual FMediaIOInputConfiguration GetDefaultInputConfiguration() const override;
+	virtual FMediaIOOutputConfiguration GetDefaultOutputConfiguration() const override;
+	virtual FMediaIOVideoTimecodeConfiguration GetDefaultTimecodeConfiguration() const override;
+	
+	virtual FText ToText(const FMediaIOConfiguration& InConfiguration, bool bInIsAutoDetected = false) const override;
+	virtual FText ToText(const FMediaIOConnection& InConnection) const override;
+	virtual FText ToText(const FMediaIOOutputConfiguration& InConfiguration) const override;
+
+#if WITH_EDITOR
+	virtual bool ShowInputTransportInSelector() const override { return false; }
+	virtual bool ShowInputKeyInSelector() const override { return false; }
+	virtual bool ShowOutputKeyInSelector() const override { return false; }
+#endif
+
+	//~ End IMediaIOCoreDeviceProvider
+};
diff --git a/Engine/Plugins/Media/NDIMedia/Source/NDIMedia/Private/NDIMediaCapture.cpp b/Engine/Plugins/Media/NDIMedia/Source/NDIMedia/Private/NDIMediaCapture.cpp
--- a/Engine/Plugins/Media/NDIMedia/Source/NDIMedia/Private/NDIMediaCapture.cpp
+++ b/Engine/Plugins/Media/NDIMedia/Source/NDIMedia/Private/NDIMediaCapture.cpp
@@ -6,12 +6,11 @@
 #include "Engine/TextureRenderTarget2D.h"
 #include "MediaIOCoreSubsystem.h"
 #include "NDIMediaAPI.h"
+#include "NDIMediaLog.h"
 #include "NDIMediaModule.h"
 #include "NDIMediaOutput.h"
 #include "Slate/SceneViewport.h"
 
-DEFINE_LOG_CATEGORY(LogNDIMedia);
-
 class UNDIMediaCapture::FNDICaptureInstance
 {
 public:
@@ -218,7 +217,7 @@ public:
 
 		// Take a local copy of the audio output in case it is switched from the main thread.
 		const TSharedPtr<FMediaIOAudioOutput> LocalAudioOutput = AudioOutput;
-		if(!LocalAudioOutput)
+		if (!LocalAudioOutput)
 		{
 			return;
 		}
diff --git a/Engine/Plugins/Media/NDIMedia/Source/NDIMedia/Private/NDIMediaLog.h b/Engine/Plugins/Media/NDIMedia/Source/NDIMedia/Private/NDIMediaLog.h
new file mode 100644
--- /dev/null
+++ b/Engine/Plugins/Media/NDIMedia/Source/NDIMedia/Private/NDIMediaLog.h
@@ -0,0 +1,7 @@
+// Copyright Epic Games, Inc. All Rights Reserved.
+ 
+#pragma once
+ 
+#include "Logging/LogMacros.h"
+ 
+DECLARE_LOG_CATEGORY_EXTERN(LogNDIMedia, Log, All);
\ No newline at end of file
diff --git a/Engine/Plugins/Media/NDIMedia/Source/NDIMedia/Private/NDIMediaModule.cpp b/Engine/Plugins/Media/NDIMedia/Source/NDIMedia/Private/NDIMediaModule.cpp
--- a/Engine/Plugins/Media/NDIMedia/Source/NDIMedia/Private/NDIMediaModule.cpp
+++ b/Engine/Plugins/Media/NDIMedia/Source/NDIMedia/Private/NDIMediaModule.cpp
@@ -4,13 +4,21 @@
 
 #include "GenericPlatform/GenericPlatformMisc.h"
 #include "HAL/PlatformProcess.h"
+#include "IMediaIOCoreModule.h"
+#include "IMediaModule.h"
 #include "Interfaces/IPluginManager.h"
 #include "Misc/MessageDialog.h"
 #include "Misc/Paths.h"
 #include "Modules/ModuleManager.h"
+#include "NDIDeviceProvider.h"
 #include "NDIMediaAPI.h"
-#include "NDIMediaCapture.h"
+#include "NDIMediaLog.h"
 #include "NDIMediaSettings.h"
+#include "NDISourceFinder.h"
+#include "Player/NDIMediaPlayer.h"
+#include "Player/NDIStreamReceiverManager.h"
+
+DEFINE_LOG_CATEGORY(LogNDIMedia);
 
 #define LOCTEXT_NAMESPACE "NDIMediaModule"
 
@@ -87,6 +95,12 @@ FNDIMediaRuntimeLibrary::~FNDIMediaRuntimeLibrary()
 
 void FNDIMediaModule::StartupModule()
 {
+	// supported platforms
+	SupportedPlatforms.Add(TEXT("Windows"));
+
+	// supported schemes
+	SupportedUriSchemes.Add(TEXT("ndi")); // Also in NdiDeviceProvider.cpp
+	
 #if UE_EDITOR
 	if (UNDIMediaSettings* Settings = GetMutableDefault<UNDIMediaSettings>())
 	{
@@ -94,10 +108,22 @@ void FNDIMediaModule::StartupModule()
 	}
 #endif
 
+	// register player factory
+    if (IMediaModule* MediaModule = FModuleManager::LoadModulePtr<IMediaModule>("Media"))
+    {
+    	MediaModule->RegisterPlayerFactory(*this);
+    }
+
 	if (!LoadModuleDependencies())
 	{
 		UE_LOG(LogNDIMedia, Error, TEXT("Unable to load \"" NDILIB_LIBRARY_NAME "\" from the specified location(s)."));
+		return;
 	}
+
+	DeviceProvider = MakeShared<FNDIDeviceProvider>();
+	IMediaIOCoreModule::Get().RegisterDeviceProvider(DeviceProvider.Get());
+
+	StreamReceiverManager = MakeShared<FNDIStreamReceiverManager>();
 }
 
 void FNDIMediaModule::ShutdownModule()
@@ -110,9 +136,105 @@ void FNDIMediaModule::ShutdownModule()
 	}
 #endif	
 
+	// unregister player factory
+	if (IMediaModule* MediaModule = FModuleManager::GetModulePtr<IMediaModule>("Media"))
+	{
+		MediaModule->UnregisterPlayerFactory(*this);
+	}
+
+	if (IMediaIOCoreModule::IsAvailable())
+	{
+		IMediaIOCoreModule::Get().UnregisterDeviceProvider(DeviceProvider.Get());
+	}
+	DeviceProvider.Reset();
+
+	StreamReceiverManager.Reset();
+	
 	NDILib.Reset();
 }
 
+TSharedPtr<FNDISourceFinder> FNDIMediaModule::GetFindInstance()
+{
+	if (!FindInstance)
+	{
+		FindInstance = MakeShared<FNDISourceFinder>(NDILib);
+	}
+	else
+	{
+		FindInstance->Validate(NDILib);
+	}
+	return FindInstance;
+}
+
+bool FNDIMediaModule::CanPlayUrl(const FString& Url, const IMediaOptions* Options, TArray<FText>* OutWarnings, TArray<FText>* OutErrors) const
+{
+	return GetPlayabilityConfidenceScore(Url, Options, OutWarnings, OutErrors) > 0 ? true : false;
+}
+
+int32 FNDIMediaModule::GetPlayabilityConfidenceScore(const FString& Url, const IMediaOptions* Options, TArray<FText>* OutWarnings, TArray<FText>* OutErrors) const
+{
+	FString Scheme;
+	FString Location;
+
+	// check scheme
+	if (!Url.Split(TEXT("://"), &Scheme, &Location, ESearchCase::CaseSensitive))
+	{
+		if (OutErrors != nullptr)
+		{
+			OutErrors->Add(LOCTEXT("NoSchemeFound", "No URI scheme found"));
+		}
+
+		return 0;
+	}
+
+	if (!SupportedUriSchemes.Contains(Scheme))
+	{
+		if (OutErrors != nullptr)
+		{
+			OutErrors->Add(FText::Format(LOCTEXT("SchemeNotSupported", "The URI scheme '{0}' is not supported"), FText::FromString(Scheme)));
+		}
+
+		return 0;
+	}
+
+	return 100;
+}
+
+TSharedPtr<IMediaPlayer> FNDIMediaModule::CreatePlayer(IMediaEventSink& InEventSink)
+{
+	return MakeShared<FNDIMediaPlayer>(InEventSink);
+}
+
+FText FNDIMediaModule::GetDisplayName() const
+{
+	return LOCTEXT("MediaPlayerFactory_DisplayName", "NDI");
+}
+
+FName FNDIMediaModule::GetPlayerName() const
+{
+	static FName PlayerName(TEXT("NDIMedia"));
+	return PlayerName;
+}
+
+const FGuid FNDIMediaModule::PlayerPluginGUID(0xc25ed21c, 0x1a7f4320, 0x8e898ae5, 0xb0f6aea2);
+
+FGuid FNDIMediaModule::GetPlayerPluginGUID() const
+{
+	return PlayerPluginGUID;
+}
+
+const TArray<FString>& FNDIMediaModule::GetSupportedPlatforms() const
+{
+	return SupportedPlatforms;
+}
+
+bool FNDIMediaModule::SupportsFeature(EMediaFeature Feature) const
+{
+	return Feature == EMediaFeature::AudioSamples ||
+			Feature == EMediaFeature::MetadataTracks ||
+			Feature == EMediaFeature::VideoSamples;
+}
+
 namespace UE::NDIMedia::Private
 {
 	static const TCHAR* DefaultLibraryName = TEXT(NDILIB_LIBRARY_NAME);
@@ -161,7 +283,7 @@ bool FNDIMediaModule::LoadModuleDependencies()
 	bool bIsLoaded = NDILib.IsValid() && NDILib->IsLoaded();
 
 	// Fallback to bundled library if something was wrong with system one.
-	if(!bIsLoaded && !Settings->bUseBundledLibrary)
+	if (!bIsLoaded && !Settings->bUseBundledLibrary)
 	{
 		LibraryPath = GetRuntimeLibraryFullPath();
 		UE_LOG(LogNDIMedia, Warning, TEXT("Falling back to bundled NDI runtime library: \"%s\"."), *LibraryPath);
diff --git a/Engine/Plugins/Media/NDIMedia/Source/NDIMedia/Private/NDIMediaModule.h b/Engine/Plugins/Media/NDIMedia/Source/NDIMedia/Private/NDIMediaModule.h
--- a/Engine/Plugins/Media/NDIMedia/Source/NDIMedia/Private/NDIMediaModule.h
+++ b/Engine/Plugins/Media/NDIMedia/Source/NDIMedia/Private/NDIMediaModule.h
@@ -2,11 +2,22 @@
 
 #pragma once
 
+#include "IMediaPlayerFactory.h"
+#include "Modules/ModuleInterface.h"
 #include "Modules/ModuleManager.h"
+#include "Templates/SharedPointer.h"
 
+class FNDIDeviceProvider;
+class FNDISourceFinder;
+class FNDIStreamReceiverManager;
+class IMediaEventSink;
+class IMediaPlayer;
 class UNDIMediaSettings;
 struct NDIlib_v5;
 
+/**
+ * Wrapper for the loaded NDI runtime library.
+ */
 class FNDIMediaRuntimeLibrary
 {
 public:
@@ -24,33 +35,67 @@ public:
 	/** Handle to the NDI runtime dll. */
 	void* LibHandle = nullptr;
 
+	/** Path the library was loaded from. */
 	FString LibraryPath;
 
 	/** Keep track of senders being created to detect source collisions and provide better error messages. */
 	TSet<FString> Senders;	// Format: GroupName_SourceName
 };
 
-class FNDIMediaModule : public IModuleInterface
+class FNDIMediaModule : public IModuleInterface, public IMediaPlayerFactory
 {
 public:
 	/** IModuleInterface implementation */
 	virtual void StartupModule() override;
 	virtual void ShutdownModule() override;
 
+	/** Returns the module, if loaded or null otherwise */
+	static FNDIMediaModule* Get()
+	{
+		return FModuleManager::GetModulePtr<FNDIMediaModule>(TEXT("NDIMedia"));
+	}
+
 	/**
 	 * Returns a handle to the currently loaded NDI runtime library.
 	 * Objects holding runtime resources should also keep a ref on the library.
 	 */
 	static TSharedPtr<FNDIMediaRuntimeLibrary> GetNDIRuntimeLibrary()
 	{
-		FNDIMediaModule* NDIMediaModule = FModuleManager::GetModulePtr<FNDIMediaModule>(TEXT("NDIMedia"));
-		if (NDIMediaModule != nullptr)
+		if (FNDIMediaModule* NDIMediaModule = FNDIMediaModule::Get())
 		{
 			return NDIMediaModule->NDILib;
 		}
 		return nullptr;
 	}
-	
+
+	/** Returns NDI source find instance. */
+	TSharedPtr<FNDISourceFinder> GetFindInstance();
+
+	/** Access the ndi device provider */
+	TSharedPtr<FNDIDeviceProvider> GetDeviceProvider() const
+	{
+		return DeviceProvider;
+	}
+
+	/** Access the ndi stream receiver manager. */
+	FNDIStreamReceiverManager& GetStreamReceiverManager() const
+	{
+		return *StreamReceiverManager;
+	}
+
+	//~ Begin IMediaPlayerFactory
+	virtual bool CanPlayUrl(const FString& Url, const IMediaOptions* Options, TArray<FText>* OutWarnings, TArray<FText>* OutErrors) const override;
+	virtual int32 GetPlayabilityConfidenceScore(const FString& Url, const IMediaOptions* Options, TArray<FText>* OutWarnings, TArray<FText>* OutErrors) const override;
+	virtual TSharedPtr<IMediaPlayer, ESPMode::ThreadSafe> CreatePlayer(IMediaEventSink& InEventSink) override;
+	virtual FText GetDisplayName() const override;
+	virtual FName GetPlayerName() const override;
+	virtual FGuid GetPlayerPluginGUID() const override;
+	virtual const TArray<FString>& GetSupportedPlatforms() const override;
+	virtual bool SupportsFeature(EMediaFeature Feature) const override;
+	//~ End IMediaPlayerFactory
+
+	static const FGuid PlayerPluginGUID;
+
 private:
 	bool LoadModuleDependencies();
 
@@ -60,4 +105,14 @@ private:
 #endif
 	
 	TSharedPtr<FNDIMediaRuntimeLibrary> NDILib;
+
+	/** List of platforms that the media player support. */
+	TArray<FString> SupportedPlatforms;
+
+	/** List of supported URI schemes. */
+	TArray<FString> SupportedUriSchemes;
+
+	TSharedPtr<FNDISourceFinder> FindInstance;
+	TSharedPtr<FNDIDeviceProvider> DeviceProvider;
+	TSharedPtr<FNDIStreamReceiverManager> StreamReceiverManager;
 };
diff --git a/Engine/Plugins/Media/NDIMedia/Source/NDIMedia/Private/NDIMediaOutput.cpp b/Engine/Plugins/Media/NDIMedia/Source/NDIMedia/Private/NDIMediaOutput.cpp
--- a/Engine/Plugins/Media/NDIMedia/Source/NDIMedia/Private/NDIMediaOutput.cpp
+++ b/Engine/Plugins/Media/NDIMedia/Source/NDIMedia/Private/NDIMediaOutput.cpp
@@ -3,10 +3,12 @@
 #include "NDIMediaOutput.h"
 
 #include "NDIMediaCapture.h"
+#include "NDIMediaLog.h"
 #include "UnrealEngine.h"
 
 UNDIMediaOutput::UNDIMediaOutput(const FObjectInitializer& ObjectInitializer)
 	: Super(ObjectInitializer)
+	, DesiredSize(0,0)
 {
 }
 
diff --git a/Engine/Plugins/Media/NDIMedia/Source/NDIMedia/Private/NDIMediaSource.cpp b/Engine/Plugins/Media/NDIMedia/Source/NDIMedia/Private/NDIMediaSource.cpp
new file mode 100644
--- /dev/null
+++ b/Engine/Plugins/Media/NDIMedia/Source/NDIMedia/Private/NDIMediaSource.cpp
@@ -0,0 +1,231 @@
+// Copyright Epic Games, Inc. All Rights Reserved.
+
+#include "NDIMediaSource.h"
+
+#include "MediaIOCorePlayerBase.h"
+#include "NDIDeviceProvider.h"
+#include "NDIMediaLog.h"
+#include "NDIMediaModule.h"
+#include "Player/NDIMediaSourceOptions.h"
+
+UNDIMediaSource::UNDIMediaSource()
+	: UCaptureCardMediaSource()
+	, bCaptureAncillary(false)
+	, MaxNumAncillaryFrameBuffer(8)
+	, bCaptureAudio(false)
+	, MaxNumAudioFrameBuffer(8)
+	, bCaptureVideo(true)
+	, MaxNumVideoFrameBuffer(8)
+	, bLogDropFrame(true)
+	, bEncodeTimecodeInTexel(false)
+{
+	bOverrideSourceEncoding = false;
+	OverrideSourceEncoding = EMediaIOCoreSourceEncoding::sRGB;
+	bOverrideSourceColorSpace = false;
+	OverrideSourceColorSpace = ETextureColorSpace::TCS_sRGB;
+	
+	AssignDefaultConfiguration();
+}
+
+bool UNDIMediaSource::GetMediaOption(const FName& Key, bool DefaultValue) const
+{
+	if (Key == UE::NDIMediaSourceOptions::SyncTimecodeToSource)
+	{
+		return bSyncTimecodeToSource;
+	}
+	if (Key == UE::NDIMediaSourceOptions::CaptureAncillary)
+	{
+		return bCaptureAncillary;
+	}
+	if (Key == UE::NDIMediaSourceOptions::CaptureAudio)
+	{
+		return bCaptureAudio;
+	}
+	if (Key == UE::NDIMediaSourceOptions::CaptureVideo)
+	{
+		return bCaptureVideo;
+	}
+	if (Key == UE::NDIMediaSourceOptions::LogDropFrame)
+	{
+		return bLogDropFrame;
+	}
+	if (Key == UE::NDIMediaSourceOptions::EncodeTimecodeInTexel)
+	{
+		return bEncodeTimecodeInTexel;
+	}
+
+	return Super::GetMediaOption(Key, DefaultValue);
+}
+
+int64 UNDIMediaSource::GetMediaOption(const FName& Key, int64 DefaultValue) const
+{
+	if (Key == FMediaIOCoreMediaOption::FrameRateNumerator)
+	{
+		return MediaConfiguration.MediaMode.FrameRate.Numerator;
+	}
+	if (Key == FMediaIOCoreMediaOption::FrameRateDenominator)
+	{
+		return MediaConfiguration.MediaMode.FrameRate.Denominator;
+	}
+	if (Key == FMediaIOCoreMediaOption::ResolutionWidth)
+	{
+		return MediaConfiguration.MediaMode.Resolution.X;
+	}
+	if (Key == FMediaIOCoreMediaOption::ResolutionHeight)
+	{
+		return MediaConfiguration.MediaMode.Resolution.Y;
+	}
+	if (Key == UE::NDIMediaSourceOptions::Bandwidth)
+	{
+		return static_cast<int64>(Bandwidth);
+	}
+	if (Key == UE::NDIMediaSourceOptions::MaxAncillaryFrameBuffer)
+	{
+		return MaxNumAncillaryFrameBuffer;
+	}
+	if (Key == UE::NDIMediaSourceOptions::MaxAudioFrameBuffer)
+	{
+		return MaxNumAudioFrameBuffer;
+	}
+	if (Key == UE::NDIMediaSourceOptions::MaxVideoFrameBuffer)
+	{
+		return MaxNumVideoFrameBuffer;
+	}
+
+	return Super::GetMediaOption(Key, DefaultValue);
+}
+
+FString UNDIMediaSource::GetMediaOption(const FName& Key, const FString& DefaultValue) const
+{
+	if (Key == UE::NDIMediaSourceOptions::DeviceName)
+	{
+		return MediaConfiguration.MediaConnection.Device.DeviceName.ToString();
+	}
+	if (Key == FMediaIOCoreMediaOption::VideoModeName)
+	{
+		return MediaConfiguration.MediaMode.GetModeName().ToString();
+	}
+	return Super::GetMediaOption(Key, DefaultValue);
+}
+
+bool UNDIMediaSource::HasMediaOption(const FName& Key) const
+{
+	return Super::HasMediaOption(Key)
+		|| Key == FMediaIOCoreMediaOption::FrameRateNumerator
+		|| Key == FMediaIOCoreMediaOption::FrameRateDenominator
+		|| Key == FMediaIOCoreMediaOption::ResolutionWidth
+		|| Key == FMediaIOCoreMediaOption::ResolutionHeight
+		|| Key == FMediaIOCoreMediaOption::VideoModeName
+		|| Key == UE::NDIMediaSourceOptions::DeviceName
+		|| Key == UE::NDIMediaSourceOptions::Bandwidth
+		|| Key == UE::NDIMediaSourceOptions::SyncTimecodeToSource
+		|| Key == UE::NDIMediaSourceOptions::LogDropFrame
+		|| Key == UE::NDIMediaSourceOptions::EncodeTimecodeInTexel
+		|| Key == UE::NDIMediaSourceOptions::CaptureAudio
+		|| Key == UE::NDIMediaSourceOptions::CaptureVideo
+		|| Key == UE::NDIMediaSourceOptions::CaptureAncillary
+		|| Key == UE::NDIMediaSourceOptions::MaxAudioFrameBuffer
+		|| Key == UE::NDIMediaSourceOptions::MaxVideoFrameBuffer
+		|| Key == UE::NDIMediaSourceOptions::MaxAncillaryFrameBuffer;
+}
+
+FString UNDIMediaSource::GetUrl() const
+{
+	// The device name contains the full ndi source name (machine (stream) format).
+	return FString::Printf(TEXT("ndi://%s"), *MediaConfiguration.MediaConnection.Device.DeviceName.ToString());
+}
+
+bool UNDIMediaSource::Validate() const
+{
+	FString FailureReason;
+	if (bAutoDetectInput)
+	{
+		if (!MediaConfiguration.MediaConnection.IsValid())
+		{
+			UE_LOG(LogNDIMedia, Warning, TEXT("The MediaConfiguration '%s' is invalid."), *GetName());
+			return false;
+		}
+	}
+	else
+	{
+		if (!MediaConfiguration.IsValid())
+		{
+			UE_LOG(LogNDIMedia, Warning, TEXT("The MediaConfiguration '%s' is invalid."), *GetName());
+			return false;
+		}
+	}
+
+	if (!bRenderJIT && EvaluationType == EMediaIOSampleEvaluationType::Latest)
+	{
+		UE_LOG(LogNDIMedia, Warning, TEXT("The MediaSource '%s' uses 'Latest' evaluation type which requires JIT rendering."), *GetName());
+		return false;
+	}
+
+	if (bFramelock)
+	{
+		UE_LOG(LogNDIMedia, Warning, TEXT("The MediaSource '%s' uses 'Framelock' which has not been implemented yet. This option will be ignored."), *GetName());
+	}
+
+	FNDIMediaModule* NdiModule = FNDIMediaModule::Get();
+	if (!NdiModule)
+	{
+		UE_LOG(LogNDIMedia, Error, TEXT("The MediaSource '%s' failed to validate because the Ndi Module is not loaded."), *GetName());
+		return false;
+	}
+
+	const TSharedPtr<FNDIDeviceProvider> DeviceProvider = NdiModule->GetDeviceProvider();
+	if (!DeviceProvider)
+	{
+		UE_LOG(LogNDIMedia, Error, TEXT("The MediaSource '%s' failed to validate because the Ndi Device Provider is not created."), *GetName());
+		return false;
+	}
+
+	bool bDeviceFound = false;
+	TArray<FMediaIOInputConfiguration> InputConfigurations = DeviceProvider->GetInputConfigurations();
+	for (const FMediaIOInputConfiguration& InputConfiguration : InputConfigurations)
+	{
+		if (InputConfiguration.MediaConfiguration.MediaConnection.Device.DeviceName == MediaConfiguration.MediaConnection.Device.DeviceName)
+		{
+			bDeviceFound = true;
+			break;
+		}
+	}
+		
+	if (!bDeviceFound)
+	{
+		UE_LOG(LogNDIMedia, Warning, TEXT("The MediaSource '%s' use the device '%s' that doesn't exist on this machine."), *GetName(), *MediaConfiguration.MediaConnection.Device.DeviceName.ToString());
+		return false;
+	}
+	
+	return true;
+}
+
+void UNDIMediaSource::PostLoad()
+{
+	Super::PostLoad();
+
+	AssignDefaultConfiguration();
+}
+
+void UNDIMediaSource::AssignDefaultConfiguration()
+{
+	if (!MediaConfiguration.IsValid())
+	{
+		if (const FNDIMediaModule* NdiModule = FNDIMediaModule::Get())
+		{
+			if (const TSharedPtr<FNDIDeviceProvider> DeviceProvider = NdiModule->GetDeviceProvider())
+			{
+				const TArray<FMediaIOConfiguration> Configurations = DeviceProvider->GetConfigurations();
+				for (const FMediaIOConfiguration& Configuration : Configurations)
+				{
+					if (Configuration.bIsInput)
+					{
+						MediaConfiguration = Configuration;
+						bRenderJIT = false;
+						break;
+					}
+				}
+			}
+		}
+	}
+}
diff --git a/Engine/Plugins/Media/NDIMedia/Source/NDIMedia/Private/NDISourceFinder.cpp b/Engine/Plugins/Media/NDIMedia/Source/NDIMedia/Private/NDISourceFinder.cpp
new file mode 100644
--- /dev/null
+++ b/Engine/Plugins/Media/NDIMedia/Source/NDIMedia/Private/NDISourceFinder.cpp
@@ -0,0 +1,62 @@
+// Copyright Epic Games, Inc. All Rights Reserved.
+
+#include "NDISourceFinder.h"
+
+#include "NDIMediaAPI.h"
+#include "NDIMediaModule.h"
+
+FNDISourceFinder::FNDISourceFinder(const TSharedPtr<FNDIMediaRuntimeLibrary>& InNdiLib)
+	: NdiLib(InNdiLib)
+{
+	Create();
+}
+
+FNDISourceFinder::~FNDISourceFinder()
+{
+	Destroy();
+}
+
+/** Call prior to doing other calls to ensure the runtime lib is up to date. */
+void FNDISourceFinder::Validate(const TSharedPtr<FNDIMediaRuntimeLibrary>& InNdiLib)
+{
+	if (NdiLib->LibHandle != InNdiLib->LibHandle)
+	{
+		Destroy();
+		NdiLib = InNdiLib;
+		Create();
+	}
+}
+
+TArray<FNDISourceFinder::FNDISourceInfo> FNDISourceFinder::GetSources() const
+{
+	FScopeLock Lock(&FindSyncContext);
+		
+	uint32 NumSources = 0;
+	const NDIlib_source_t* Sources = NdiLib->Lib->find_get_current_sources(FindInstance, &NumSources);
+
+	TArray<FNDISourceInfo> OutSources;
+	OutSources.Reserve(NumSources);
+	for (uint32 i = 0; i < NumSources; ++i)
+	{
+		FNDISourceInfo SourceInfo;
+		SourceInfo.Name = Sources[i].p_ndi_name;
+		SourceInfo.Url = Sources[i].p_url_address;
+		OutSources.Add(SourceInfo);
+	}
+
+	return OutSources;
+}
+
+void FNDISourceFinder::Create()
+{
+	FindInstance = NdiLib->Lib->find_create_v2(nullptr);
+}
+
+void FNDISourceFinder::Destroy()
+{
+	if (FindInstance && NdiLib && NdiLib->IsLoaded())
+	{
+		NdiLib->Lib->find_destroy(FindInstance);
+		FindInstance = nullptr;
+	}
+}
diff --git a/Engine/Plugins/Media/NDIMedia/Source/NDIMedia/Private/NDISourceFinder.h b/Engine/Plugins/Media/NDIMedia/Source/NDIMedia/Private/NDISourceFinder.h
new file mode 100644
--- /dev/null
+++ b/Engine/Plugins/Media/NDIMedia/Source/NDIMedia/Private/NDISourceFinder.h
@@ -0,0 +1,39 @@
+// Copyright Epic Games, Inc. All Rights Reserved.
+
+#pragma once
+
+#include "NDIMediaAPI.h"
+
+class FNDIMediaRuntimeLibrary;
+
+/**
+ *  Wrapper for the NDI source finder instance.
+ */
+class FNDISourceFinder
+{
+public:
+	FNDISourceFinder(const TSharedPtr<FNDIMediaRuntimeLibrary>& InNdiLib);
+	~FNDISourceFinder();
+
+	/** Call prior to doing other calls to ensure the runtime lib is up to date. */
+	void Validate(const TSharedPtr<FNDIMediaRuntimeLibrary>& InNdiLib);
+
+	/** NDI Source information */
+	struct FNDISourceInfo
+	{
+		FString Name;
+		FString Url;
+	};
+
+	/** Returns current list of sources discovered on the network. */
+	TArray<FNDISourceInfo> GetSources() const;
+
+private:
+	void Create();
+	void Destroy();
+	
+private:
+	TSharedPtr<FNDIMediaRuntimeLibrary> NdiLib;
+	NDIlib_find_instance_t FindInstance = nullptr;
+	mutable FCriticalSection FindSyncContext;
+};
diff --git a/Engine/Plugins/Media/NDIMedia/Source/NDIMedia/Private/NDITimecodeProvider.cpp b/Engine/Plugins/Media/NDIMedia/Source/NDIMedia/Private/NDITimecodeProvider.cpp
new file mode 100644
--- /dev/null
+++ b/Engine/Plugins/Media/NDIMedia/Source/NDIMedia/Private/NDITimecodeProvider.cpp
@@ -0,0 +1,110 @@
+// Copyright Epic Games, Inc. All Rights Reserved.
+
+#include "NDITimecodeProvider.h"
+
+#include "NDIMediaModule.h"
+#include "Player/NDIStreamReceiver.h"
+#include "Player/NDIStreamReceiverManager.h"
+
+bool UNDITimecodeProvider::FetchTimecode(FQualifiedFrameTime& OutFrameTime)
+{
+	FScopeLock Lock(&StateSyncContext);
+
+	if (!Receiver ||
+	    (GetSynchronizationState() != ETimecodeProviderSynchronizationState::Synchronized))
+	{
+		return false;
+	}
+
+	OutFrameTime = MostRecentFrameTime;
+
+	return true;
+}
+
+ETimecodeProviderSynchronizationState UNDITimecodeProvider::GetSynchronizationState() const
+{
+	FScopeLock Lock(&StateSyncContext);
+
+	if (!Receiver)
+	{
+		return ETimecodeProviderSynchronizationState::Closed;
+	}
+
+	return State;
+}
+
+bool UNDITimecodeProvider::Initialize(UEngine* InEngine)
+{
+	State = ETimecodeProviderSynchronizationState::Closed;
+
+	FNDISourceSettings SourceSettings;
+	SourceSettings.Bandwidth = Bandwidth;
+	SourceSettings.bCaptureAudio = false;
+	SourceSettings.bCaptureVideo = true;
+	SourceSettings.SourceName = TimecodeConfiguration.MediaConfiguration.MediaConnection.Device.DeviceName.ToString();
+
+	// Check if the receiver is already created by another object.
+	if (FNDIMediaModule* Module = FNDIMediaModule::Get())
+	{
+		FNDIStreamReceiverManager& StreamReceiverManager = Module->GetStreamReceiverManager();
+		Receiver = StreamReceiverManager.FindReceiver(SourceSettings.SourceName);
+		if (Receiver)
+		{
+			SourceSettings.bCaptureAudio = Receiver->GetCurrentSourceSettings().bCaptureAudio;
+		}
+	}
+
+	if (!Receiver)
+	{
+		Receiver = MakeShared<FNDIStreamReceiver>(FNDIMediaModule::GetNDIRuntimeLibrary());	
+	}
+
+	Receiver->SetSyncTimecodeToSource(/* bInSyncTimecodeToSource*/true);
+	Receiver->Initialize(SourceSettings, FNDIStreamReceiver::ECaptureMode::OnEndFrameRT);
+
+	VideoFrameReceivedHandle = Receiver->OnVideoFrameReceived.AddLambda([this](FNDIStreamReceiver* InReceiver, const NDIlib_video_frame_v2_t& InVideoFrame, const FTimespan& InTimecode)
+	{
+		FScopeLock Lock(&StateSyncContext);
+		State = ETimecodeProviderSynchronizationState::Synchronized;
+		MostRecentFrameTime = FQualifiedFrameTime(Receiver->GetCurrentTimecode(), Receiver->GetCurrentFrameRate());
+	});
+	ConnectedHandle = Receiver->OnConnected.AddLambda([this](FNDIStreamReceiver* InReceiver)
+	{
+		FScopeLock Lock(&StateSyncContext);
+		State = ETimecodeProviderSynchronizationState::Synchronizing;
+	});
+	DisconnectedHandle = Receiver->OnDisconnected.AddLambda([this](FNDIStreamReceiver* InReceiver)
+	{
+		FScopeLock Lock(&StateSyncContext);
+		State = ETimecodeProviderSynchronizationState::Closed;
+	});
+
+	return true;
+}
+
+void UNDITimecodeProvider::Shutdown(UEngine* InEngine)
+{
+	ReleaseResources();
+}
+
+void UNDITimecodeProvider::BeginDestroy()
+{
+	ReleaseResources();
+
+	Super::BeginDestroy();
+}
+
+void UNDITimecodeProvider::ReleaseResources()
+{
+	if (Receiver)
+	{
+		Receiver->OnVideoFrameReceived.Remove(VideoFrameReceivedHandle);
+		Receiver->OnConnected.Remove(ConnectedHandle);
+		Receiver->OnDisconnected.Remove(DisconnectedHandle);
+	}
+	VideoFrameReceivedHandle.Reset();
+	ConnectedHandle.Reset();
+	DisconnectedHandle.Reset();
+
+	State = ETimecodeProviderSynchronizationState::Closed;
+}
diff --git a/Engine/Plugins/Media/NDIMedia/Source/NDIMedia/Private/Player/NDIMediaPlayer.cpp b/Engine/Plugins/Media/NDIMedia/Source/NDIMedia/Private/Player/NDIMediaPlayer.cpp
new file mode 100644
--- /dev/null
+++ b/Engine/Plugins/Media/NDIMedia/Source/NDIMedia/Private/Player/NDIMediaPlayer.cpp
@@ -0,0 +1,481 @@
+// Copyright Epic Games, Inc. All Rights Reserved.
+
+#include "NDIMediaPlayer.h"
+
+#include "IMediaEventSink.h"
+#include "IMediaOptions.h"
+#include "MediaIOCoreAudioSampleBase.h"
+#include "MediaIOCoreBinarySampleBase.h"
+#include "MediaIOCoreEncodeTime.h"
+#include "MediaIOCoreSamples.h"
+#include "MediaIOCoreTextureSampleBase.h"
+#include "NDIMediaLog.h"
+#include "NDIMediaModule.h"
+#include "NDIMediaSourceOptions.h"
+#include "NDIStreamReceiver.h"
+#include "NDIStreamReceiverManager.h"
+
+#define LOCTEXT_NAMESPACE "NDIMediaPlayer"
+
+/**
+ * Implements a media texture sample for NDIMedia.
+ */
+class FNDIMediaTextureSample : public FMediaIOCoreTextureSampleBase
+{
+	using Super = FMediaIOCoreTextureSampleBase;
+
+public:
+	/**
+	 * Initialize the sample.
+	 * 
+	 * @param InVideoFrame Received Video frame data from NDI.
+	 * @param InColorFormatArgs Color space and encoding settings from the media source.
+	 * @param InTime Current playback time of the sample.
+	 * @return true if the sample was correctly initialized, false otherwise.
+	 */
+	bool Initialize(const NDIlib_video_frame_v2_t& InVideoFrame, const UE::MediaIOCore::FColorFormatArgs& InColorFormatArgs, const FTimespan& InTime)
+	{
+		int32 FrameBufferSize;
+		EMediaTextureSampleFormat FrameSampleFormat;
+
+		if (InVideoFrame.FourCC == NDIlib_FourCC_video_type_UYVY)
+		{
+			FrameBufferSize = InVideoFrame.line_stride_in_bytes * InVideoFrame.yres;
+			FrameSampleFormat = EMediaTextureSampleFormat::CharUYVY;
+		}
+		else if (InVideoFrame.FourCC == NDIlib_FourCC_video_type_BGRA)
+		{
+			FrameBufferSize = InVideoFrame.line_stride_in_bytes * InVideoFrame.yres;
+			FrameSampleFormat = EMediaTextureSampleFormat::CharBGRA;
+		}
+		else if (InVideoFrame.FourCC == NDIlib_FourCC_video_type_RGBA
+			|| InVideoFrame.FourCC == NDIlib_FourCC_video_type_RGBX)
+		{
+			FrameBufferSize = InVideoFrame.line_stride_in_bytes * InVideoFrame.yres;
+			FrameSampleFormat = EMediaTextureSampleFormat::CharRGBA;
+		}
+		// else if (InVideoFrame.FourCC == NDIlib_FourCC_video_type_UYVA)
+		// {
+		// 	FrameBufferSize = InVideoFrame.line_stride_in_bytes * InVideoFrame.yres + InVideoFrame.xres*InVideoFrame.yres;
+		// 	FrameSampleFormat = EMediaTextureSampleFormat::CharUYVA; // todo
+		// }
+		else
+		{
+			return false;
+		}
+
+		switch (InVideoFrame.frame_format_type)
+		{
+			case NDIlib_frame_format_type_progressive:
+				return Super::Initialize(InVideoFrame.p_data
+					, FrameBufferSize
+					, InVideoFrame.line_stride_in_bytes
+					, InVideoFrame.xres
+					, InVideoFrame.yres
+					, FrameSampleFormat
+					, InTime
+					, FFrameRate(InVideoFrame.frame_rate_N, InVideoFrame.frame_rate_D)
+					, TOptional<FTimecode>() // todo
+					, InColorFormatArgs);		
+				
+			case NDIlib_frame_format_type_field_0:
+			case NDIlib_frame_format_type_field_1:
+			
+			return Super::InitializeWithEvenOddLine(InVideoFrame.frame_format_type == NDIlib_frame_format_type_field_0
+				, InVideoFrame.p_data
+				, FrameBufferSize
+				, InVideoFrame.line_stride_in_bytes
+				, InVideoFrame.xres
+				, InVideoFrame.yres
+				, FrameSampleFormat
+				, InTime
+				, FFrameRate(InVideoFrame.frame_rate_N, InVideoFrame.frame_rate_D)
+				, TOptional<FTimecode>() // todo
+				, InColorFormatArgs);
+
+			default:
+				return false;
+		}
+	}
+};
+
+class FNDIMediaTextureSamplePool : public TMediaObjectPool<FNDIMediaTextureSample>
+{};
+
+
+/**
+ * Implements a media audio sample for NDIMedia
+ */
+class FNDIMediaAudioSample : public FMediaIOCoreAudioSampleBase
+{
+	using Super = FMediaIOCoreAudioSampleBase;
+
+public:
+};
+
+class FNDIMediaAudioSamplePool : public TMediaObjectPool<FNDIMediaAudioSample>
+{};
+
+/*
+ * Implements a pool for NDI binary sample objects. 
+ */
+class FNDIMediaBinarySamplePool : public TMediaObjectPool<FMediaIOCoreBinarySampleBase> { };
+
+FNDIMediaPlayer::FNDIMediaPlayer(IMediaEventSink& InEventSink)
+	: Super(InEventSink)
+	, NDIPlayerState(EMediaState::Closed)
+	, EventSink(InEventSink)
+	, TextureSamplePool(new FNDIMediaTextureSamplePool)
+	, AudioSamplePool(new FNDIMediaAudioSamplePool)
+	, MetadataSamplePool(new FNDIMediaBinarySamplePool)
+{}
+
+FNDIMediaPlayer::~FNDIMediaPlayer()
+{
+	Close();
+
+	TextureSamplePool.Reset();
+	AudioSamplePool.Reset();
+}
+
+FGuid FNDIMediaPlayer::GetPlayerPluginGUID() const
+{
+	return FNDIMediaModule::PlayerPluginGUID;
+}
+
+bool FNDIMediaPlayer::Open(const FString& InUrl, const IMediaOptions* InOptions)
+{
+	if (!Super::Open(InUrl, InOptions))
+	{
+		return false;
+	}
+
+	MaxNumVideoFrameBuffer = InOptions->GetMediaOption(UE::NDIMediaSourceOptions::MaxVideoFrameBuffer, (int64)8);
+	MaxNumAudioFrameBuffer = InOptions->GetMediaOption(UE::NDIMediaSourceOptions::MaxAudioFrameBuffer, (int64)8);
+	MaxNumMetadataFrameBuffer = InOptions->GetMediaOption(UE::NDIMediaSourceOptions::MaxAncillaryFrameBuffer, (int64)8);
+	bEncodeTimecodeInTexel = InOptions->GetMediaOption(UE::NDIMediaSourceOptions::EncodeTimecodeInTexel, false);
+
+	// Setup our different supported channels based on source settings
+	SetupSampleChannels();
+
+	// configure format information for base class
+	AudioTrackFormat.BitsPerSample = 32;
+	AudioTrackFormat.NumChannels = 0;
+	AudioTrackFormat.SampleRate = 44100;
+	AudioTrackFormat.TypeName = FString(TEXT("PCM"));
+
+	bCaptureVideo = InOptions->GetMediaOption(UE::NDIMediaSourceOptions::CaptureVideo, true);
+	bCaptureAudio = InOptions->GetMediaOption(UE::NDIMediaSourceOptions::CaptureAudio, false);
+	bCaptureAncillary = InOptions->GetMediaOption(UE::NDIMediaSourceOptions::CaptureAncillary, false);
+	SupportedSampleTypes = bCaptureVideo ? EMediaIOSampleType::Video : EMediaIOSampleType::None;
+	SupportedSampleTypes |= bCaptureAudio ? EMediaIOSampleType::Audio : EMediaIOSampleType::None;
+	SupportedSampleTypes |= bCaptureAncillary ? EMediaIOSampleType::Metadata : EMediaIOSampleType::None;
+	Samples->EnableTimedDataChannels(this, SupportedSampleTypes);
+
+	FNDISourceSettings SourceSettings;
+	SourceSettings.Bandwidth = static_cast<ENDIReceiverBandwidth>(InOptions->GetMediaOption(UE::NDIMediaSourceOptions::Bandwidth, static_cast<int64>(SourceSettings.Bandwidth)));
+	SourceSettings.bCaptureAudio = bCaptureAudio;
+	SourceSettings.bCaptureVideo = bCaptureVideo;
+
+	FString Scheme;
+	FString Location;
+	if (InUrl.Split(TEXT("://"), &Scheme, &Location, ESearchCase::CaseSensitive))
+	{
+		SourceSettings.SourceName = Location;
+	}
+
+	// Check if the receiver is already created by another object.
+	if (FNDIMediaModule* Module = FNDIMediaModule::Get())
+	{
+		FNDIStreamReceiverManager& StreamReceiverManager = Module->GetStreamReceiverManager();
+		Receiver = StreamReceiverManager.FindReceiver(SourceSettings.SourceName);
+
+		if (!Receiver)
+		{
+			Receiver = MakeShared<FNDIStreamReceiver>(FNDIMediaModule::GetNDIRuntimeLibrary());
+		}
+	}
+	
+	if (!Receiver)
+	{
+		UE_LOG(LogNDIMedia, Error, TEXT("Failed to acquire NDI receiver."));
+		return false;
+	}
+	
+	// Hook into the video and audio captures
+	VideoReceivedHandle = Receiver->OnVideoFrameReceived.AddRaw(this, &FNDIMediaPlayer::HandleVideoFrameReceived);
+	AudioReceivedHandle = Receiver->OnAudioFrameReceived.AddRaw(this, &FNDIMediaPlayer::HandleAudioFrameReceived);
+
+	// Control the player's state based on the receiver connecting and disconnecting
+	ConnectedHandle = Receiver->OnConnected.AddLambda([this](FNDIStreamReceiver* receiver)
+	{
+		NDIPlayerState = EMediaState::Playing;
+	});
+	DisconnectedHandle = Receiver->OnDisconnected.AddLambda([this](FNDIStreamReceiver* receiver)
+	{
+		NDIPlayerState = EMediaState::Closed;
+	});
+
+	// Get ready to connect
+	CurrentState = EMediaState::Preparing;
+	NDIPlayerState = EMediaState::Preparing;
+	EventSink.ReceiveMediaEvent(EMediaEvent::MediaConnecting);
+	
+	Receiver->SetSyncTimecodeToSource(InOptions->GetMediaOption(UE::NDIMediaSourceOptions::SyncTimecodeToSource, true));
+
+	// Start up the receiver under the player's control.
+	return Receiver->Initialize(SourceSettings, FNDIStreamReceiver::ECaptureMode::Manual);
+}
+
+void FNDIMediaPlayer::Close()
+{
+	NDIPlayerState = EMediaState::Closed;
+
+	if (Receiver != nullptr)
+	{
+		// Disconnect from receiver events
+		Receiver->OnVideoFrameReceived.Remove(VideoReceivedHandle);
+		VideoReceivedHandle.Reset();
+		Receiver->OnAudioFrameReceived.Remove(AudioReceivedHandle);
+		AudioReceivedHandle.Reset();
+		Receiver->OnConnected.Remove(ConnectedHandle);
+		ConnectedHandle.Reset();
+		Receiver->OnDisconnected.Remove(DisconnectedHandle);
+		DisconnectedHandle.Reset();
+
+		Receiver.Reset();
+	}
+
+	TextureSamplePool->Reset();
+	AudioSamplePool->Reset();
+	MetadataSamplePool->Reset();
+
+	Super::Close();
+}
+
+FString FNDIMediaPlayer::GetStats() const
+{
+	FString Stats;
+
+	if (Receiver)
+	{
+		const FNDIReceiverPerformanceData PerformanceData = Receiver->GetPerformanceData();
+		Stats += FString::Printf(TEXT("Video Frames: %lld"), PerformanceData.VideoFrames);
+		Stats += FString::Printf(TEXT("Dropped Video Frames: %lld"), PerformanceData.DroppedVideoFrames);
+		Stats += FString::Printf(TEXT("Audio Frames: %lld"), PerformanceData.AudioFrames);
+		Stats += FString::Printf(TEXT("Dropped Audio Frames: %lld"), PerformanceData.DroppedAudioFrames);
+		Stats += FString::Printf(TEXT("Metadata Frames: %lld"), PerformanceData.MetadataFrames);
+		Stats += FString::Printf(TEXT("Dropped Metadata Frames: %lld"), PerformanceData.DroppedMetadataFrames);
+	}
+	else
+	{
+		Stats = FString(TEXT("Receiver not available."));
+	}
+	return Stats;
+}
+
+void FNDIMediaPlayer::TickInput(FTimespan InDeltaTime, FTimespan InTime)
+{
+	// Update player state
+	EMediaState NewState = NDIPlayerState;
+
+	if (NewState != CurrentState)
+	{
+		CurrentState = NewState;
+		if (CurrentState == EMediaState::Playing)
+		{
+			EventSink.ReceiveMediaEvent(EMediaEvent::TracksChanged);
+			EventSink.ReceiveMediaEvent(EMediaEvent::MediaOpened);
+			EventSink.ReceiveMediaEvent(EMediaEvent::PlaybackResumed);
+		}
+		else if (NewState == EMediaState::Error)
+		{
+			EventSink.ReceiveMediaEvent(EMediaEvent::MediaOpenFailed);
+			Close();
+		}
+	}
+
+	if (CurrentState != EMediaState::Playing)
+	{
+		return;
+	}
+
+	TickTimeManagement();
+}
+
+
+void FNDIMediaPlayer::TickFetch(FTimespan InDeltaTime, FTimespan InTime)
+{
+	Super::TickFetch(InDeltaTime, InTime);
+
+	if (CurrentState == EMediaState::Preparing || CurrentState == EMediaState::Playing)
+	{
+		if (Receiver != nullptr)
+		{
+			if (bCaptureAudio)
+			{
+				Receiver->FetchAudio(InTime);
+			}
+			if (bCaptureVideo)
+			{
+				Receiver->FetchVideo(InTime);
+			}
+			if (bCaptureAncillary)
+			{
+				// Potential improvement: limit how much metadata is processed, to avoid appearing to lock up due to a metadata flood
+				while (Receiver->FetchMetadata(InTime)) {}
+			}
+		}
+	}
+
+	if (CurrentState == EMediaState::Playing)
+	{
+		// No need to lock here. That info is only used for debug information.
+		AudioTrackFormat.NumChannels = NDIThreadAudioChannels;
+		AudioTrackFormat.SampleRate = NDIThreadAudioSampleRate;
+
+		if (Receiver)
+		{
+			VideoFrameRate = Receiver->GetCurrentFrameRate();
+			VideoTrackFormat.Dim = Receiver->GetCurrentResolution();
+			VideoTrackFormat.FrameRates = TRange<float>(VideoFrameRate.AsDecimal());
+			VideoTrackFormat.FrameRate = VideoFrameRate.AsDecimal();
+			//VideoTrackFormat.TypeName = Configuration.Configuration.MediaMode.GetModeName().ToString();
+		}
+
+		VerifyFrameDropCount();
+	}
+}
+
+void FNDIMediaPlayer::HandleVideoFrameReceived(FNDIStreamReceiver* InReceiver, const NDIlib_video_frame_v2_t& InVideoFrame, const FTimespan& InTime)
+{
+	TSharedRef<FNDIMediaTextureSample> TextureSample = TextureSamplePool->AcquireShared();
+	
+	const UE::MediaIOCore::FColorFormatArgs ColorFormatArgs(
+		bOverrideSourceEncoding ?  static_cast<UE::Color::EEncoding>(OverrideSourceEncoding) : UE::Color::EEncoding::sRGB,
+		bOverrideSourceColorSpace ? static_cast<UE::Color::EColorSpace>(OverrideSourceColorSpace) : UE::Color::EColorSpace::sRGB);
+	
+	if (TextureSample->Initialize(InVideoFrame, ColorFormatArgs, InTime))
+	{
+		TextureSample->SetColorConversionSettings(OCIOSettings);
+
+		if (bEncodeTimecodeInTexel && InVideoFrame.frame_format_type == NDIlib_frame_format_type_progressive)
+		{
+			EMediaIOCoreEncodePixelFormat EncodePixelFormat;
+			bool bEncodeSupported = true;
+			if (InVideoFrame.FourCC == NDIlib_FourCC_video_type_UYVY)
+			{
+				EncodePixelFormat = EMediaIOCoreEncodePixelFormat::CharUYVY;
+			}
+			else if (InVideoFrame.FourCC == NDIlib_FourCC_video_type_BGRA
+				|| InVideoFrame.FourCC == NDIlib_FourCC_video_type_RGBA
+				|| InVideoFrame.FourCC == NDIlib_FourCC_video_type_BGRX)
+			{
+				EncodePixelFormat = EMediaIOCoreEncodePixelFormat::CharBGRA;
+			}
+			else
+			{
+				EncodePixelFormat = EMediaIOCoreEncodePixelFormat::CharUYVY;
+				bEncodeSupported = false;
+			}
+			
+			if (bEncodeSupported)
+			{
+				FTimecode SetTimecode = InReceiver->GetCurrentTimecode();
+				FMediaIOCoreEncodeTime EncodeTime(EncodePixelFormat, const_cast<void*>(TextureSample->GetBuffer()), InVideoFrame.line_stride_in_bytes, InVideoFrame.xres, InVideoFrame.yres);
+				EncodeTime.Render(SetTimecode.Hours, SetTimecode.Minutes, SetTimecode.Seconds, SetTimecode.Frames);
+			}
+		}
+		
+		AddVideoSample(TextureSample);
+	}
+}
+
+void FNDIMediaPlayer::HandleAudioFrameReceived(FNDIStreamReceiver* InReceiver, const NDIlib_audio_frame_v2_t& InAudioFrame, const FTimespan& InTime)
+{
+	FNDIMediaRuntimeLibrary* NdiLib = Receiver->GetNdiLib();
+	if (!NdiLib || !NdiLib->IsLoaded())
+	{
+		return;
+	}
+	
+	TSharedRef<FNDIMediaAudioSample> AudioSample = AudioSamplePool->AcquireShared();
+
+	// UE wants 32bit signed interleaved audio data, so need to convert the NDI audio.
+	// Fortunately the NDI library has a utility function to do that.
+
+	// Get a buffer to convert to
+	const int32 available_samples = InAudioFrame.no_samples * InAudioFrame.no_channels;
+	void* SampleBuffer = AudioSample->RequestBuffer(available_samples);
+
+	if (SampleBuffer != nullptr)
+	{
+		// Format to convert to
+		NDIlib_audio_frame_interleaved_32s_t audio_frame_32s(
+			InAudioFrame.sample_rate,
+			InAudioFrame.no_channels,
+			InAudioFrame.no_samples,
+			InAudioFrame.timecode,
+			20,
+			static_cast<int32_t*>(SampleBuffer));
+
+		// Convert received NDI audio
+		NdiLib->Lib->util_audio_to_interleaved_32s_v2(&InAudioFrame, &audio_frame_32s);
+		
+		// Supply converted audio data
+		if (AudioSample->SetProperties(available_samples
+			, audio_frame_32s.no_channels
+			, audio_frame_32s.sample_rate
+			, InTime
+			, TOptional<FTimecode>()))
+		{
+			NDIThreadAudioChannels = audio_frame_32s.no_channels;
+			NDIThreadAudioSampleRate = audio_frame_32s.sample_rate;
+			
+			AddAudioSample(AudioSample);
+		}
+	}
+}
+
+void FNDIMediaPlayer::VerifyFrameDropCount()
+{
+	// todo
+}
+
+bool FNDIMediaPlayer::IsHardwareReady() const
+{
+	return NDIPlayerState == EMediaState::Playing;
+}
+
+void FNDIMediaPlayer::SetupSampleChannels()
+{
+	FMediaIOSamplingSettings VideoSettings = BaseSettings;
+	VideoSettings.BufferSize = MaxNumVideoFrameBuffer;
+	Samples->InitializeVideoBuffer(VideoSettings);
+
+	FMediaIOSamplingSettings AudioSettings = BaseSettings;
+	AudioSettings.BufferSize = MaxNumAudioFrameBuffer;
+	Samples->InitializeAudioBuffer(AudioSettings);
+
+	FMediaIOSamplingSettings MetadataSettings = BaseSettings;
+	MetadataSettings.BufferSize = MaxNumMetadataFrameBuffer;
+	Samples->InitializeMetadataBuffer(MetadataSettings);
+}
+
+
+TSharedPtr<FMediaIOCoreTextureSampleBase> FNDIMediaPlayer::AcquireTextureSample_AnyThread() const
+{
+	// Needed by the deinterlacer.
+	return TextureSamplePool->AcquireShared();
+}
+
+//~ ITimedDataInput interface
+#if WITH_EDITOR
+const FSlateBrush* FNDIMediaPlayer::GetDisplayIcon() const
+{
+	return nullptr;
+}
+#endif
+
+
+#undef LOCTEXT_NAMESPACE
diff --git a/Engine/Plugins/Media/NDIMedia/Source/NDIMedia/Private/Player/NDIMediaPlayer.h b/Engine/Plugins/Media/NDIMedia/Source/NDIMedia/Private/Player/NDIMediaPlayer.h
new file mode 100644
--- /dev/null
+++ b/Engine/Plugins/Media/NDIMedia/Source/NDIMedia/Private/Player/NDIMediaPlayer.h
@@ -0,0 +1,102 @@
+// Copyright Epic Games, Inc. All Rights Reserved.
+
+#pragma once
+
+#include "MediaIOCorePlayerBase.h"
+#include "NDIMediaAPI.h"
+
+class FNDIMediaAudioSamplePool;
+class FNDIMediaBinarySamplePool;
+class FNDIMediaRuntimeLibrary;
+class FNDIMediaTextureSamplePool;
+class FNDIStreamReceiver;
+
+/**
+ * Implementation of the Media player for an NDI stream.
+ */
+class FNDIMediaPlayer : public FMediaIOCorePlayerBase
+{
+	using Super = FMediaIOCorePlayerBase;
+
+public:
+	FNDIMediaPlayer(IMediaEventSink& InEventSink);
+
+	virtual ~FNDIMediaPlayer() override;
+
+	//~ Begin IMediaPlayer
+	virtual FGuid GetPlayerPluginGUID() const override;
+	virtual bool Open(const FString& InUrl, const IMediaOptions* InOptions) override;
+	virtual void Close() override;
+	virtual FString GetStats() const override;
+	virtual void TickInput(FTimespan InDeltaTime, FTimespan InTime) override;
+	virtual void TickFetch(FTimespan InDeltaTime, FTimespan InTime) override;
+
+protected:
+	virtual bool IsHardwareReady() const override;
+	virtual void SetupSampleChannels() override;
+	virtual TSharedPtr<FMediaIOCoreTextureSampleBase> AcquireTextureSample_AnyThread() const override;
+	//~ End IMediaPlayer
+
+public:
+	//~ Begin ITimedDataInput
+#if WITH_EDITOR
+	virtual const FSlateBrush* GetDisplayIcon() const override;
+#endif
+	//~ End ITimedDataInput
+
+private:
+	/** Callback handler for video frame reception. */
+	void HandleVideoFrameReceived(FNDIStreamReceiver* InReceiver, const NDIlib_video_frame_v2_t& InVideoFrame, const FTimespan& InTime);
+
+	/** Callback handler for audio frame reception. */
+	void HandleAudioFrameReceived(FNDIStreamReceiver* InReceiver, const NDIlib_audio_frame_v2_t& InAudioFrame, const FTimespan& InTime);
+
+	void VerifyFrameDropCount();
+
+private:
+	/** Maximum Audio Sample Pool size */
+	int32 MaxNumAudioFrameBuffer = 0;
+	
+	/** Maximum Metadata Sample Pool size */
+	int32 MaxNumMetadataFrameBuffer = 0;
+
+	/** Maximum Video Sample Pool size */
+	int32 MaxNumVideoFrameBuffer = 0;
+
+	/** Determines which of the sample streams to capture (mirrors MediaSource). */
+	bool bCaptureVideo = true;
+	bool bCaptureAudio = false;
+	bool bCaptureAncillary = false;
+
+	/** Whether to use the time code embedded in video frames. */
+	bool bEncodeTimecodeInTexel = false;
+
+	/** Used to flag which sample types we advertise as supported for timed data monitoring */
+	EMediaIOSampleType SupportedSampleTypes = EMediaIOSampleType::None;
+
+	/** Current state of the media player. */
+	EMediaState NDIPlayerState = EMediaState::Closed;
+
+	/** Number of channels from the last audio frame received. */
+	uint32 NDIThreadAudioChannels = 0;
+
+	/** Sample rate from the last audio frame received. */
+	uint32 NDIThreadAudioSampleRate = 0;
+	
+	/** The media event sink. */
+	IMediaEventSink& EventSink;
+
+	/** Current stream receiver */
+	TSharedPtr<FNDIStreamReceiver> Receiver;
+
+	/** Handles for the receiver delegates */ 
+	FDelegateHandle VideoReceivedHandle;
+	FDelegateHandle AudioReceivedHandle;
+	FDelegateHandle ConnectedHandle;
+	FDelegateHandle DisconnectedHandle;
+
+	/** Media Sample pools */
+	TUniquePtr<FNDIMediaTextureSamplePool> TextureSamplePool;
+	TUniquePtr<FNDIMediaAudioSamplePool> AudioSamplePool;
+	TUniquePtr<FNDIMediaBinarySamplePool> MetadataSamplePool;
+};
diff --git a/Engine/Plugins/Media/NDIMedia/Source/NDIMedia/Private/Player/NDIMediaSourceOptions.h b/Engine/Plugins/Media/NDIMedia/Source/NDIMedia/Private/Player/NDIMediaSourceOptions.h
new file mode 100644
--- /dev/null
+++ b/Engine/Plugins/Media/NDIMedia/Source/NDIMedia/Private/Player/NDIMediaSourceOptions.h
@@ -0,0 +1,21 @@
+// Copyright Epic Games, Inc. All Rights Reserved.
+
+#pragma once
+
+#include "UObject/NameTypes.h"
+
+/** NDI Media source specific options exposed to the player. */
+namespace UE::NDIMediaSourceOptions
+{
+	static const FLazyName DeviceName("DeviceName");
+	static const FLazyName Bandwidth("Bandwidth");
+	static const FLazyName SyncTimecodeToSource("SyncTimecodeToSource");
+	static const FLazyName LogDropFrame("LogDropFrame");
+	static const FLazyName EncodeTimecodeInTexel("EncodeTimecodeInTexel");
+	static const FLazyName CaptureAudio("CaptureAudio");
+	static const FLazyName CaptureVideo("CaptureVideo");
+	static const FLazyName CaptureAncillary("CaptureAncillary");
+	static const FLazyName MaxVideoFrameBuffer("MaxVideoFrameBuffer");
+	static const FLazyName MaxAudioFrameBuffer("MaxAudioFrameBuffer");
+	static const FLazyName MaxAncillaryFrameBuffer("MaxAncillaryFrameBuffer");
+}
diff --git a/Engine/Plugins/Media/NDIMedia/Source/NDIMedia/Private/Player/NDIReceiverPerformanceData.h b/Engine/Plugins/Media/NDIMedia/Source/NDIMedia/Private/Player/NDIReceiverPerformanceData.h
new file mode 100644
--- /dev/null
+++ b/Engine/Plugins/Media/NDIMedia/Source/NDIMedia/Private/Player/NDIReceiverPerformanceData.h
@@ -0,0 +1,48 @@
+// Copyright Epic Games, Inc. All Rights Reserved.
+
+#pragma once
+
+#include "NDIReceiverPerformanceData.generated.h"
+
+/**
+ *  NDI receiver's performance statistics on number of frames received and dropped.
+ */
+USTRUCT(BlueprintType)
+struct FNDIReceiverPerformanceData
+{
+	GENERATED_BODY()
+
+public:
+	/** The number of video frames received */
+	UPROPERTY(BlueprintReadOnly, VisibleAnywhere, Category = "Information")
+	int64 VideoFrames = 0;
+
+	/** The number of video frames dropped */
+	UPROPERTY(BlueprintReadonly, VisibleAnywhere, Category = "Information")
+	int64 DroppedVideoFrames = 0;
+
+	/** The number of audio frames received */
+	UPROPERTY(BlueprintReadOnly, VisibleAnywhere, Category = "Information")
+	int64 AudioFrames = 0;
+
+	/** The number of audio frames dropped */
+	UPROPERTY(BlueprintReadOnly, VisibleAnywhere, Category = "Information")
+	int64 DroppedAudioFrames = 0;
+
+	/** The number of metadata frames received */
+	UPROPERTY(BlueprintReadOnly, VisibleAnywhere, Category = "Information")
+	int64 MetadataFrames = 0;
+
+	/** The number of metadata frames dropped */
+	UPROPERTY(BlueprintReadOnly, VisibleAnywhere, Category = "Information")
+	int64 DroppedMetadataFrames = 0;
+
+	/** Compares this object to 'other' and returns a determination of whether they are equal */
+	bool operator==(const FNDIReceiverPerformanceData& other) const
+	{
+		return AudioFrames == other.AudioFrames && DroppedAudioFrames == other.DroppedAudioFrames &&
+			   DroppedMetadataFrames == other.DroppedMetadataFrames &&
+			   DroppedVideoFrames == other.DroppedVideoFrames && MetadataFrames == other.MetadataFrames &&
+			   VideoFrames == other.VideoFrames;
+	}
+};
\ No newline at end of file
diff --git a/Engine/Plugins/Media/NDIMedia/Source/NDIMedia/Private/Player/NDISourceSettings.h b/Engine/Plugins/Media/NDIMedia/Source/NDIMedia/Private/Player/NDISourceSettings.h
new file mode 100644
--- /dev/null
+++ b/Engine/Plugins/Media/NDIMedia/Source/NDIMedia/Private/Player/NDISourceSettings.h
@@ -0,0 +1,45 @@
+// Copyright Epic Games, Inc. All Rights Reserved.
+
+#pragma once
+
+#include "NDIMediaDefines.h"
+
+#include "NDISourceSettings.generated.h"
+
+/**
+ * Carries the NDI source and capture settings to the media receiver
+ * from either media source/player or timecode provider. 
+ */
+USTRUCT(BlueprintType)
+struct FNDISourceSettings
+{
+	GENERATED_BODY()
+
+public:
+	/** NDI Source name */
+	UPROPERTY(BlueprintReadWrite, EditAnywhere, Category="Properties")
+	FString SourceName = FString("");
+
+	/** Indicates the current bandwidth mode used for the connection to this source */
+	UPROPERTY(BlueprintReadWrite, EditAnywhere, Category="Properties")
+	ENDIReceiverBandwidth Bandwidth = ENDIReceiverBandwidth::Highest;
+
+	/** Capture Audio from the NDI source. */
+	UPROPERTY(BlueprintReadWrite, EditAnywhere, Category="Properties")
+	bool bCaptureAudio = false;
+
+	/** Capture Video from the NDI source. */
+	UPROPERTY(BlueprintReadWrite, EditAnywhere, Category="Properties")
+	bool bCaptureVideo = false;
+
+	bool IsValid() const
+	{
+		return !SourceName.IsEmpty();
+	}
+
+	bool operator==(const FNDISourceSettings& InOther) const
+	{
+		return this->Bandwidth == InOther.Bandwidth && this->SourceName == InOther.SourceName &&
+			   this->bCaptureVideo == InOther.bCaptureVideo && this->bCaptureAudio == InOther.bCaptureAudio;
+	}
+};
diff --git a/Engine/Plugins/Media/NDIMedia/Source/NDIMedia/Private/Player/NDIStreamReceiver.cpp b/Engine/Plugins/Media/NDIMedia/Source/NDIMedia/Private/Player/NDIStreamReceiver.cpp
new file mode 100644
--- /dev/null
+++ b/Engine/Plugins/Media/NDIMedia/Source/NDIMedia/Private/Player/NDIStreamReceiver.cpp
@@ -0,0 +1,560 @@
+// Copyright Epic Games, Inc. All Rights Reserved.
+
+#include "Player/NDIStreamReceiver.h"
+
+#include "Async/Async.h"
+#include "Misc/CoreDelegates.h"
+#include "NDIMediaAPI.h"
+#include "NDIMediaLog.h"
+#include "NDIMediaModule.h"
+
+#if WITH_EDITOR
+#include "Editor.h"
+#endif
+
+namespace UE::NDIStreamReceiver::Private
+{
+	NDIlib_recv_bandwidth_e GetNDIBandwidth(ENDIReceiverBandwidth InBandwidth)
+	{
+		switch (InBandwidth)
+		{
+		case ENDIReceiverBandwidth::Highest:
+			return NDIlib_recv_bandwidth_highest;
+		case ENDIReceiverBandwidth::MetadataOnly:
+			return NDIlib_recv_bandwidth_metadata_only;
+		case ENDIReceiverBandwidth::AudioOnly:
+			return NDIlib_recv_bandwidth_audio_only;
+		case ENDIReceiverBandwidth::Lowest:
+			return NDIlib_recv_bandwidth_lowest;
+		default:
+			return NDIlib_recv_bandwidth_highest;
+		}
+	}
+
+	FTimecode ToTimecode(const FTimespan& InTime, const FFrameRate& InFrameRate)
+	{
+		const bool bDropFrame = FTimecode::IsDropFormatTimecodeSupported(InFrameRate);
+		constexpr bool bRollOver = true; // use roll-over timecode
+		return FTimecode::FromTimespan(InTime, InFrameRate, bDropFrame, bRollOver);
+	}
+	
+	FTimecode ToTimecode(int64 InSourceTicks, const FFrameRate& InFrameRate)
+	{
+		FTimespan Time = FTimespan::FromSeconds(InSourceTicks / (double)1e+7);
+		return ToTimecode(Time, InFrameRate);
+	}
+}
+
+FNDIStreamReceiver::FNDIStreamReceiver(const TSharedPtr<FNDIMediaRuntimeLibrary>& InNdiLib)
+	: Resolution(0, 0)
+	, NdiLib(InNdiLib)
+{
+	
+}
+
+FNDIStreamReceiver::~FNDIStreamReceiver()
+{
+	Shutdown();
+}
+
+bool FNDIStreamReceiver::Initialize(const FNDISourceSettings& InSourceSettings, ECaptureMode InCaptureMode)
+{
+	if (!ReceiveInstance)
+	{
+		// create a non-connected receiver instance
+		NDIlib_recv_create_v3_t ReceiveSettings;
+		ReceiveSettings.allow_video_fields = false;
+		ReceiveSettings.bandwidth = NDIlib_recv_bandwidth_highest;
+		ReceiveSettings.color_format = NDIlib_recv_color_format_fastest;
+
+		ReceiveInstance = NdiLib->Lib->recv_create_v3(&ReceiveSettings);
+		
+		if (!ReceiveInstance)
+		{
+			UE_LOG(LogNDIMedia, Error, TEXT("Failed to create NDI receiver."));
+			return false;
+		}
+	}
+
+	// If the incoming source settings are valid
+	if (InSourceSettings.IsValid())
+	{
+		// Connect to the source with the new settings.
+		ChangeConnection(InSourceSettings);
+	}
+
+	if (InCaptureMode == ECaptureMode::OnEndFrameRT)
+	{
+		// We don't want to limit the engine rendering speed to the sync rate of the connection hook
+		// into the core delegates render thread 'EndFrame'
+		FCoreDelegates::OnEndFrameRT.Remove(FrameEndRTHandle);
+		FrameEndRTHandle.Reset();
+		FrameEndRTHandle = FCoreDelegates::OnEndFrameRT.AddLambda([this]()
+		{
+			const FTimespan Time = FTimespan::FromSeconds(FPlatformTime::Seconds());
+			
+			while (FetchMetadata(Time))
+				; // Potential improvement: limit how much metadata is processed, to avoid appearing to lock up due to a metadata flood
+			FetchVideo(Time);
+		});
+
+#if UE_EDITOR
+		// We don't want to provide perceived issues with the plugin not working so
+		// when we get a Pre-exit message, forcefully shutdown the receiver
+		FCoreDelegates::OnPreExit.AddSPLambda(this, [This=this]()
+		{
+			This->Shutdown();
+			FCoreDelegates::OnPreExit.RemoveAll(This);
+		});
+
+		// We handle this in the 'Play In Editor' versions as well.
+		FEditorDelegates::PrePIEEnded.AddSPLambda(this, [This=this](const bool)
+		{
+			This->Shutdown();
+			FEditorDelegates::PrePIEEnded.RemoveAll(This);
+		});
+#endif
+	}
+
+	return true;
+}
+
+void FNDIStreamReceiver::StartConnection()
+{
+	FScopeLock RenderLock(&RenderSyncContext);
+	FScopeLock AudioLock(&AudioSyncContext);
+	FScopeLock MetadataLock(&MetadataSyncContext);
+
+	if (SourceSettings.IsValid())
+	{
+		// Create a non-connected receiver instance
+		NDIlib_recv_create_v3_t ReceiveSettings;
+		ReceiveSettings.allow_video_fields = true;
+		ReceiveSettings.bandwidth = UE::NDIStreamReceiver::Private::GetNDIBandwidth(SourceSettings.Bandwidth);
+		ReceiveSettings.color_format = NDIlib_recv_color_format_fastest;
+
+		// Do the conversion on the connection information
+		NDIlib_source_t SourceInfo;
+
+		const FUtf8String SourceNameUtf8 = StringCast<UTF8CHAR>(*SourceSettings.SourceName).Get();
+		const FUtf8String UrlUtf8; // = StringCast<UTF8CHAR>(*SourceSettings.Url).Get(); // This is not carried from the device provider.
+		SourceInfo.p_ndi_name = reinterpret_cast<const char*>(*SourceNameUtf8);
+		SourceInfo.p_url_address = reinterpret_cast<const char*>(*UrlUtf8);
+
+		// Create a receiver and connect to the source
+		NDIlib_recv_instance_type* NewReceiveInstance = NdiLib->Lib->recv_create_v3(&ReceiveSettings);
+		NdiLib->Lib->recv_connect(NewReceiveInstance, &SourceInfo);
+
+		// Get rid of existing connection
+		StopConnection();
+
+		// set the receiver to the new connection
+		ReceiveInstance = NewReceiveInstance;
+
+		// create a new frame sync instance
+		FrameSyncInstance = NdiLib->Lib->framesync_create(ReceiveInstance);
+	}
+}
+
+void FNDIStreamReceiver::StopConnection()
+{
+	FScopeLock RenderLock(&RenderSyncContext);
+	FScopeLock AudioLock(&AudioSyncContext);
+	FScopeLock MetadataLock(&MetadataSyncContext);
+
+	// destroy the framesync instance
+	if (FrameSyncInstance != nullptr)
+	{
+		NdiLib->Lib->framesync_destroy(FrameSyncInstance);
+		FrameSyncInstance = nullptr;
+	}
+
+	// Free the receiver
+	if (ReceiveInstance != nullptr)
+	{
+		NdiLib->Lib->recv_destroy(ReceiveInstance);
+		ReceiveInstance = nullptr;
+	}
+}
+
+void FNDIStreamReceiver::ChangeConnection(const FNDISourceSettings& InSourceSettings)
+{
+	// Ensure some thread-safety because FetchVideo function is called on the render thread
+	FScopeLock RenderLock(&RenderSyncContext);
+	FScopeLock AudioLock(&AudioSyncContext);
+	FScopeLock MetadataLock(&MetadataSyncContext);
+
+	// We should only worry about connections that are already created
+	if (ReceiveInstance != nullptr)
+	{
+		// Set the connection information for the requested new connection
+		if (SourceSettings != InSourceSettings)
+		{
+			bool bSettingsChanged = false;
+			if (SourceSettings.SourceName != InSourceSettings.SourceName
+				|| SourceSettings.Bandwidth != InSourceSettings.Bandwidth)
+			{
+				bSettingsChanged = true;
+			}
+			
+			SourceSettings = InSourceSettings;
+
+			if (SourceSettings.IsValid())
+			{
+				if (bSettingsChanged || !ReceiveInstance || !FrameSyncInstance)
+				{
+					// Connection information is valid, and something has changed that requires the connection to be remade
+					StartConnection();
+				}
+			}
+			else
+			{
+				// Requested connection is invalid, indicating we should close the current connection
+				StopConnection();
+			}
+		}
+	}
+}
+
+int32 FNDIStreamReceiver::GetAudioChannels()
+{
+	FScopeLock Lock(&AudioSyncContext);
+
+	int32 NumChannels = 0;
+
+	if (FrameSyncInstance && SourceSettings.bCaptureAudio)
+	{
+		int AvailableNumFrames = NdiLib->Lib->framesync_audio_queue_depth(FrameSyncInstance);	// Samples per channel
+
+		if (AvailableNumFrames > 0)
+		{
+			NDIlib_audio_frame_v2_t AudioFrame;
+			NdiLib->Lib->framesync_capture_audio(FrameSyncInstance, &AudioFrame, 0, 0, 0);
+			NumChannels = AudioFrame.no_channels;
+		}
+	}
+
+	return NumChannels;
+}
+
+void FNDIStreamReceiver::SetSyncTimecodeToSource(bool bInSyncTimecodeToSource)
+{
+	bSyncTimecodeToSource = bInSyncTimecodeToSource;
+}
+
+void FNDIStreamReceiver::SendMetadataFrame(const FString& Data)
+{
+	FScopeLock Lock(&MetadataSyncContext);
+
+	if (ReceiveInstance != nullptr)
+	{
+		NDIlib_metadata_frame_t Metadata;
+		std::string DataStr(TCHAR_TO_UTF8(*Data));
+		Metadata.p_data = const_cast<char*>(DataStr.c_str());
+		Metadata.length = DataStr.length();
+		Metadata.timecode = FDateTime::Now().GetTimeOfDay().GetTicks();
+
+		NdiLib->Lib->recv_send_metadata(ReceiveInstance, &Metadata);
+	}
+}
+
+void FNDIStreamReceiver::SendMetadataFrameAttr(const FString& InElement, const FString& InElementData)
+{
+	FString Data = "<" + InElement + ">" + InElementData + "</" + InElement + ">";
+	SendMetadataFrame(Data);
+}
+
+void FNDIStreamReceiver::SendMetadataFrameAttrs(const FString& InElement, const TMap<FString,FString>& InAttributes)
+{
+	FString Data = "<" + InElement;
+	
+	for (const TPair<FString, FString>& Attribute : InAttributes)
+	{
+		Data += " " + Attribute.Key + "=\"" + Attribute.Value + "\"";
+		Data += " " + Attribute.Key + "=\"" + Attribute.Value + "\"";
+	}
+
+	Data += "/>";
+
+	SendMetadataFrame(Data);
+}
+
+void FNDIStreamReceiver::Shutdown()
+{
+	// Unregister render thread frame end delegate lambda.
+	FCoreDelegates::OnEndFrameRT.Remove(FrameEndRTHandle);
+	FrameEndRTHandle.Reset();
+
+	{
+		FScopeLock RenderLock(&RenderSyncContext);
+		FScopeLock AudioLock(&AudioSyncContext);
+		FScopeLock MetadataLock(&MetadataSyncContext);
+
+		if (ReceiveInstance != nullptr)
+		{
+			if (FrameSyncInstance != nullptr)
+			{
+				NdiLib->Lib->framesync_destroy(FrameSyncInstance);
+				FrameSyncInstance = nullptr;
+			}
+
+			NdiLib->Lib->recv_destroy(ReceiveInstance);
+			ReceiveInstance = nullptr;
+		}
+	}
+
+	// Reset the connection status of this object
+	SetIsCurrentlyConnected(/*bInConnected*/ false, /*bInDelayBroadcastEvents*/ false);
+
+	SourceSettings = FNDISourceSettings();
+	PerformanceData = FNDIReceiverPerformanceData();
+	FrameRate = FFrameRate(60, 1);
+	Resolution = FIntPoint(0, 0);
+	Timecode = FTimecode(0, FrameRate, true, true);
+}
+
+bool FNDIStreamReceiver::FetchVideo(const FTimespan& InTime)
+{
+	// This function is called on the Engine's Main Rendering Thread. Be very careful when doing stuff here.
+	// Make sure things are done quick and efficient.
+
+	// Ensure thread safety
+	FScopeLock Lock(&RenderSyncContext);
+
+	bool bFrameReceived = false;
+
+	// check for our frame sync object and that we are actually connected to the end point
+	if (FrameSyncInstance && SourceSettings.bCaptureVideo)
+	{
+		// Using a frame-sync we can always get data which is the magic and it will adapt
+		// to the frame-rate that it is being called with.
+		NDIlib_video_frame_v2_t VideoFrame;
+		NdiLib->Lib->framesync_capture_video(FrameSyncInstance, &VideoFrame, NDIlib_frame_format_type_progressive);
+
+		// Update our Performance Metrics
+		GatherPerformanceMetrics();
+
+		if (VideoFrame.p_data)
+		{
+			// Ensure that we inform all those interested when the stream starts up
+			SetIsCurrentlyConnected(/*bInConnected*/ true);
+
+			// Update the Framerate, if it has changed
+			FrameRate.Numerator = VideoFrame.frame_rate_N;
+			FrameRate.Denominator = VideoFrame.frame_rate_D;
+
+			// Update the Resolution
+			Resolution.X = VideoFrame.xres;
+			Resolution.Y = VideoFrame.yres;
+
+			if (bSyncTimecodeToSource)
+			{
+				// Update the timecode from the current 'SourceTime' value
+				int64_t SourceTime = VideoFrame.timecode % 864000000000; // Modulo the number of 100ns intervals in 24 hours
+				Timecode = UE::NDIStreamReceiver::Private::ToTimecode(SourceTime, FrameRate);
+			}
+			else
+			{
+				// Update the timecode from the current 'SystemTime' value
+				int64_t SystemTime = FDateTime::Now().GetTimeOfDay().GetTicks();
+				Timecode = UE::NDIStreamReceiver::Private::ToTimecode(SystemTime, FrameRate);
+				// todo: (validate) use provided engine time.
+				//Timecode = UE::NDIStreamReceiver::Private::ToTimecode(InTime, FrameRate);
+			}
+
+			// Redraw if:
+			// - timestamp is undefined, or
+			// - timestamp has changed, or
+			// - frame format type has changed (e.g. different field)
+			if ((VideoFrame.timestamp == NDIlib_recv_timestamp_undefined) ||
+				(VideoFrame.timestamp != LastFrameTimestamp) ||
+				(VideoFrame.frame_format_type != LastFrameFormatType))
+			{
+				bFrameReceived = true;
+
+				LastFrameTimestamp = VideoFrame.timestamp;
+				LastFrameFormatType = VideoFrame.frame_format_type;
+
+				OnVideoFrameReceived.Broadcast(this, VideoFrame, InTime);
+
+				if (VideoFrame.p_metadata)
+				{
+					FString Data(UTF8_TO_TCHAR(VideoFrame.p_metadata));
+					OnMetaDataReceived.Broadcast(this, Data, true);
+				}
+			}
+		}
+
+		// Release the video. You could keep the frame if you want and release it later.
+		NdiLib->Lib->framesync_free_video(FrameSyncInstance, &VideoFrame);
+	}
+
+	return bFrameReceived;
+}
+
+bool FNDIStreamReceiver::FetchAudio(const FTimespan& InTime)
+{
+	FScopeLock Lock(&AudioSyncContext);
+
+	bool bFrameReceived = false;
+
+	if (FrameSyncInstance && SourceSettings.bCaptureAudio)
+	{
+		int NumSamples = NdiLib->Lib->framesync_audio_queue_depth(FrameSyncInstance);
+
+		// Using a frame-sync we can always get data which is the magic and it will adapt
+		// to the frame-rate that it is being called with.
+		NDIlib_audio_frame_v2_t AudioFrame;
+		NdiLib->Lib->framesync_capture_audio(FrameSyncInstance, &AudioFrame, 0, 0, NumSamples);
+		
+		if (AudioFrame.p_data)
+		{
+			// Ensure that we inform all those interested when the stream starts up
+			SetIsCurrentlyConnected(/*bInConnected*/ true);
+
+			const int32 NumAvailableSamples = AudioFrame.no_samples * AudioFrame.no_channels;
+
+			if (NumAvailableSamples > 0)
+			{
+				bFrameReceived = true;
+				OnAudioFrameReceived.Broadcast(this, AudioFrame, InTime);
+			}
+		}
+
+		// Release the audio frame
+		NdiLib->Lib->framesync_free_audio(FrameSyncInstance, &AudioFrame);
+	}
+
+	return bFrameReceived;
+}
+
+bool FNDIStreamReceiver::FetchMetadata(const FTimespan& InTime)
+{
+	FScopeLock Lock(&MetadataSyncContext);
+
+	bool bFrameReceived = false;
+
+	if (ReceiveInstance != nullptr)
+	{
+		NDIlib_metadata_frame_t MetadataFrame;
+		NDIlib_frame_type_e FrameType = NdiLib->Lib->recv_capture_v3(ReceiveInstance, nullptr, nullptr, &MetadataFrame, 0);
+		if (FrameType == NDIlib_frame_type_metadata)
+		{
+			if (MetadataFrame.p_data)
+			{
+				// Ensure that we inform all those interested when the stream starts up
+				SetIsCurrentlyConnected(/*bInConnected*/ true);
+
+				if (MetadataFrame.length > 0)
+				{
+					bFrameReceived = true;
+					OnMetadataFrameReceived.Broadcast(this, MetadataFrame, InTime);
+
+					FString Data(UTF8_TO_TCHAR(MetadataFrame.p_data));
+					OnMetaDataReceived.Broadcast(this, Data, false);
+				}
+			}
+
+			NdiLib->Lib->recv_free_metadata(ReceiveInstance, &MetadataFrame);
+		}
+	}
+
+	return bFrameReceived;
+}
+
+void FNDIStreamReceiver::SetIsCurrentlyConnected(bool bInConnected, bool bInDelayBroadcastEvents)
+{
+	if (bInConnected != bIsCurrentlyConnected)
+	{
+		FScopeLock Lock(&ConnectionSyncContext);
+
+		if (bInConnected != bIsCurrentlyConnected)
+		{
+			bIsCurrentlyConnected = bInConnected;
+			if (bInDelayBroadcastEvents)
+			{
+				// Broadcast in the main thread.
+				TWeakPtr<FNDIStreamReceiver> WeakSelf = AsShared();
+				AsyncTask(ENamedThreads::GameThread, [WeakSelf, bInConnected]()
+				{
+					if (TSharedPtr<FNDIStreamReceiver> Self = WeakSelf.Pin())
+					{
+						if (bInConnected)
+						{
+							Self->OnConnected.Broadcast(Self.Get());
+						}
+						else
+						{
+							Self->OnDisconnected.Broadcast(Self.Get());
+						}
+					}
+				});
+			}
+			else
+			{
+				if (bInConnected)
+				{
+					OnConnected.Broadcast(this);
+				}
+				else
+				{
+					OnDisconnected.Broadcast(this);
+				}
+			}
+		}
+	}
+}
+
+void FNDIStreamReceiver::GatherPerformanceMetrics()
+{
+	// provide references to store the values
+	NDIlib_recv_performance_t stable_performance;
+	NDIlib_recv_performance_t dropped_performance;
+
+	// get the performance values from the SDK
+	NdiLib->Lib->recv_get_performance(ReceiveInstance, &stable_performance, &dropped_performance);
+
+	// update our structure with the updated values
+	PerformanceData.AudioFrames = stable_performance.audio_frames;
+	PerformanceData.DroppedAudioFrames = dropped_performance.audio_frames;
+	PerformanceData.DroppedMetadataFrames = dropped_performance.metadata_frames;
+	PerformanceData.DroppedVideoFrames = dropped_performance.video_frames;
+	PerformanceData.MetadataFrames = stable_performance.metadata_frames;
+	PerformanceData.VideoFrames = stable_performance.video_frames;
+}
+
+FNDIReceiverPerformanceData FNDIStreamReceiver::GetPerformanceData() const
+{
+	FScopeLock Lock(&RenderSyncContext);
+	return PerformanceData;
+}
+
+bool FNDIStreamReceiver::GetIsCurrentlyConnected() const
+{
+	// TODO: Guard
+	if (ReceiveInstance != nullptr)
+	{
+		return NdiLib->Lib->recv_get_no_connections(ReceiveInstance) > 0 ? true : false;
+	}
+	return false;
+}
+
+const FNDISourceSettings& FNDIStreamReceiver::GetCurrentSourceSettings() const
+{
+	return SourceSettings;
+}
+
+const FTimecode& FNDIStreamReceiver::GetCurrentTimecode() const
+{
+	return Timecode;
+}
+
+const FFrameRate& FNDIStreamReceiver::GetCurrentFrameRate() const
+{
+	return FrameRate;
+}
+
+const FIntPoint& FNDIStreamReceiver::GetCurrentResolution() const
+{
+	return Resolution;
+}
\ No newline at end of file
diff --git a/Engine/Plugins/Media/NDIMedia/Source/NDIMedia/Private/Player/NDIStreamReceiver.h b/Engine/Plugins/Media/NDIMedia/Source/NDIMedia/Private/Player/NDIStreamReceiver.h
new file mode 100644
--- /dev/null
+++ b/Engine/Plugins/Media/NDIMedia/Source/NDIMedia/Private/Player/NDIStreamReceiver.h
@@ -0,0 +1,208 @@
+// Copyright Epic Games, Inc. All Rights Reserved.
+
+#pragma once
+
+#include "Misc/FrameRate.h"
+#include "Misc/Timecode.h"
+#include "NDISourceSettings.h"
+#include "NDIMediaAPI.h"
+#include "NDIReceiverPerformanceData.h"
+#include "Templates/SharedPointer.h"
+
+class FNDIMediaRuntimeLibrary;
+
+/**
+ * A Media object representing the NDI Receiver for being able to receive Audio, Video, and Metadata over an NDI stream.
+ */
+class FNDIStreamReceiver : public TSharedFromThis<FNDIStreamReceiver>
+{
+public:
+	explicit FNDIStreamReceiver(const TSharedPtr<FNDIMediaRuntimeLibrary>& InNdiLib);
+	~FNDIStreamReceiver();
+
+	/** Receiver's capture mode indicate how the receiver is going to call FetchVideo/Audio. */
+	enum class ECaptureMode
+	{
+		/** The user of the receiver manually triggers capturing frames through FetchVideo/Audio() */
+		Manual,	
+		/** The receiver automatically captures frames every engine render frame (in the render thread) */
+		OnEndFrameRT,
+		// todo: CaptureThread - Fetch from a separate thread that doesn't impact the main or render threads.
+	};
+
+	/**
+	 * Initialize the stream receiver.
+	 * If the source settings are valid, it will start the connection.
+	 */
+	bool Initialize(const FNDISourceSettings& InSourceSettings, ECaptureMode InCaptureMode);
+
+	/**
+	 * Attempt to (re-)start the connection
+	 */
+	void StartConnection();
+
+	/**
+	 * Stop the connection
+	 */
+	void StopConnection();
+
+	/**
+	 * Attempts to change the connection to another NDI sender source
+	 */
+	void ChangeConnection(const FNDISourceSettings& InSourceSettings);
+
+	/** Peek the audio stream to retrieve the number of audio channels. */
+	int32 GetAudioChannels();
+
+	/** Sets whether the timecode should be synced to the Source Timecode value or the engine's. */
+	void SetSyncTimecodeToSource(bool bInSyncTimecodeToSource);
+
+	/**
+	 * This will send a metadata frame to the sender
+	 * The data is expected to be valid XML
+	 */
+	void SendMetadataFrame(const FString& InData);
+	
+	/**
+	 * This will send a metadata frame to the sender
+	 * The data will be formatted as: <Element>ElementData</Element>
+	 */
+	void SendMetadataFrameAttr(const FString& InElement, const FString& InElementData);
+	
+	/**
+	 * This will send a metadata frame to the sender
+	 * The data will be formatted as: <Element Key0="Value0" Key1="Value1" Keyn="Valuen"/>
+	 */
+	void SendMetadataFrameAttrs(const FString& InElement, const TMap<FString,FString>& InAttributes);
+
+	/**
+	 * Attempts to immediately stop receiving frames.
+	 */
+	void Shutdown();
+	
+	/**
+	 * Attempts to capture a frame from the connected source.  If a new frame is captured, broadcast it to
+	 * interested receivers through the receive event.  Returns true if new data was captured.
+	 * @param InTime Engine provided Sample time
+	 */
+	bool FetchVideo(const FTimespan& InTime);
+	bool FetchAudio(const FTimespan& InTime);
+	bool FetchMetadata(const FTimespan& InTime);
+
+	/** Returns the current framerate of the connected source */
+	const FFrameRate& GetCurrentFrameRate() const;
+
+	/** Returns the current resolution of the connected source */
+	const FIntPoint& GetCurrentResolution() const;
+
+	/** Returns the current timecode of the connected source */
+	const FTimecode& GetCurrentTimecode() const;
+
+	/** Returns the current connection information of the connected source */
+	const FNDISourceSettings& GetCurrentSourceSettings() const;
+
+	/** Returns the current performance data of the receiver while connected to the source */
+	FNDIReceiverPerformanceData GetPerformanceData() const;
+
+	/** Returns a value indicating whether this object is currently connected to the sender source */
+	bool GetIsCurrentlyConnected() const;
+
+	/** Returns the current reference to the ndi runtime library used to create the receiver */
+	FNDIMediaRuntimeLibrary* GetNdiLib() const
+	{
+		return NdiLib.Get();
+	}
+
+	/**
+	 * Delegate called when the source is connected.
+	 * @remark Called from the main thread.
+	 */
+	DECLARE_MULTICAST_DELEGATE_OneParam(FOnConnected, FNDIStreamReceiver*)
+	FOnConnected OnConnected;
+
+	/**
+	 * Delegate called when the source is disconnected
+	 * @remark Called from the main thread.
+	 */
+	DECLARE_MULTICAST_DELEGATE_OneParam(FOnDisconnected, FNDIStreamReceiver*)
+	FOnDisconnected OnDisconnected;
+
+	/** Delegate called when a video frame is received. */
+	DECLARE_TS_MULTICAST_DELEGATE_ThreeParams(FOnVideoFrameReceived, FNDIStreamReceiver*, const NDIlib_video_frame_v2_t&, const FTimespan&)
+	FOnVideoFrameReceived OnVideoFrameReceived;
+
+	/** Delegate called when an audio frame is received. */
+	DECLARE_TS_MULTICAST_DELEGATE_ThreeParams(FOnAudioFrameReceived, FNDIStreamReceiver*, const NDIlib_audio_frame_v2_t&, const FTimespan&)
+	FOnAudioFrameReceived OnAudioFrameReceived;
+
+	/** Delegate called when a metadata frame is received. */
+	DECLARE_TS_MULTICAST_DELEGATE_ThreeParams(FOnMetadataFrameReceived, FNDIStreamReceiver*, const NDIlib_metadata_frame_t&, const FTimespan&)
+	FOnMetadataFrameReceived OnMetadataFrameReceived;
+
+	/** Delegate called when metadata is received (either from a metadata frame or attached to a video frame). */
+	DECLARE_TS_MULTICAST_DELEGATE_ThreeParams(FOnMetaDataReceived, FNDIStreamReceiver* /*InReceiver*/, FString /*InData*/, bool /*bInAttachedToVideoFrame*/);
+	FOnMetaDataReceived OnMetaDataReceived;
+
+private:
+	/**
+	 * Update the connection status and broadcast the events to delegates (OnConnected or OnDisconnected).
+	 * @param bInConnected New connection status.
+	 * @param bInDelayBroadcastEvents Broadcast events in game thread (if called from render thread).
+	 */
+	void SetIsCurrentlyConnected(bool bInConnected, bool bInDelayBroadcastEvents = true);
+
+	/** Gathers the performance metrics of the connection to the remote source. */
+	void GatherPerformanceMetrics();
+
+private:
+	/** The current frame count, seconds, minutes, and hours in time-code notation */
+	FTimecode Timecode;
+
+	/** The desired number of frames (per second) for video to be displayed */
+	FFrameRate FrameRate;
+
+	/** The width and height of the last received video frame */
+	FIntPoint Resolution;
+
+	/** Indicates whether the timecode should be synced to the Source Timecode value */
+	bool bSyncTimecodeToSource = true;
+
+	/** Information describing detailed information about the sender this receiver is currently connected to */
+	FNDISourceSettings SourceSettings;
+
+	/** Information describing detailed information about the receiver performance when connected to an NDI? sender */
+	FNDIReceiverPerformanceData PerformanceData;
+
+	/** Keep track of the last VideoFrame timestamp received. Used to detect new frame. */
+	int64_t LastFrameTimestamp = 0;
+
+	/** Keep track of the last VideoFrame format type received. Used to detect new frame. */
+	NDIlib_frame_format_type_e LastFrameFormatType = NDIlib_frame_format_type_max;
+
+	/** Keep track of the current connection status with the source. */
+	bool bIsCurrentlyConnected = false;
+
+	/** Reference to the ndi runtime library used to create the receiver and frame sync instances. */
+	TSharedPtr<FNDIMediaRuntimeLibrary> NdiLib;
+
+	/** NDI receiver instance */
+	NDIlib_recv_instance_t ReceiveInstance = nullptr;
+
+	/** NDI FrameSync instance */
+	NDIlib_framesync_instance_t FrameSyncInstance = nullptr;
+
+	/** Critical section for VideoFrame processing related members. */
+	mutable FCriticalSection RenderSyncContext;
+
+	/** Critical section for AudioFrame processing related members. */
+	FCriticalSection AudioSyncContext;
+
+	/** Critical section for Metadata frame processing related members. */
+	FCriticalSection MetadataSyncContext;
+
+	/** Critical section for current connection status (bIsCurrentlyConnected). */
+	FCriticalSection ConnectionSyncContext;
+
+	/** Handle for the frame end render thread delegate. */
+	FDelegateHandle FrameEndRTHandle;
+};
\ No newline at end of file
diff --git a/Engine/Plugins/Media/NDIMedia/Source/NDIMedia/Private/Player/NDIStreamReceiverManager.cpp b/Engine/Plugins/Media/NDIMedia/Source/NDIMedia/Private/Player/NDIStreamReceiverManager.cpp
new file mode 100644
--- /dev/null
+++ b/Engine/Plugins/Media/NDIMedia/Source/NDIMedia/Private/Player/NDIStreamReceiverManager.cpp
@@ -0,0 +1,40 @@
+// Copyright Epic Games, Inc. All Rights Reserved.
+
+#include "NDIStreamReceiverManager.h"
+
+#include "NDIStreamReceiver.h"
+
+TSharedPtr<FNDIStreamReceiver> FNDIStreamReceiverManager::FindReceiver(const FString& InSourceName)
+{
+	for (const TWeakPtr<FNDIStreamReceiver>& ReceiverWeak : ReceiversWeak)
+	{
+		if (TSharedPtr<FNDIStreamReceiver> Receiver = ReceiverWeak.Pin())
+		{
+			if (Receiver->GetCurrentSourceSettings().SourceName == InSourceName)
+			{
+				return Receiver;
+			}
+		}
+	}
+	return nullptr;
+}
+
+bool FNDIStreamReceiverManager::AddReceiver(const TSharedPtr<FNDIStreamReceiver>& InReceiver)
+{
+	RemoveExpiredEntries();
+	if (InReceiver && !FindReceiver(InReceiver->GetCurrentSourceSettings().SourceName))
+	{
+		ReceiversWeak.Add(InReceiver);
+		return true;
+	}
+	return false;
+}
+
+void FNDIStreamReceiverManager::RemoveExpiredEntries()
+{
+	// Cleanup the invalid entries.
+	ReceiversWeak.RemoveAll([](const TWeakPtr<FNDIStreamReceiver>& InReceiverWeak)
+	{
+		return !InReceiverWeak.IsValid();
+	});
+}
\ No newline at end of file
diff --git a/Engine/Plugins/Media/NDIMedia/Source/NDIMedia/Private/Player/NDIStreamReceiverManager.h b/Engine/Plugins/Media/NDIMedia/Source/NDIMedia/Private/Player/NDIStreamReceiverManager.h
new file mode 100644
--- /dev/null
+++ b/Engine/Plugins/Media/NDIMedia/Source/NDIMedia/Private/Player/NDIStreamReceiverManager.h
@@ -0,0 +1,38 @@
+// Copyright Epic Games, Inc. All Rights Reserved.
+
+#pragma once
+
+#include "Containers/Map.h"
+#include "Templates/SharedPointer.h"
+
+class FNDIStreamReceiver;
+
+/**
+ * Implementation of an NDI Stream manager
+ *
+ * NDI Stream receivers can be shared between media players or time code providers.
+ * This manager simply allows for the receivers to be shared.
+ */
+class FNDIStreamReceiverManager : public TSharedFromThis<FNDIStreamReceiverManager>
+{
+public:
+	/**
+	 * Finds a managed receiver for the given source.
+	 * @param InSourceName Source name of the receiver to find.
+	 * @return Found receiver, or null if none.
+	 */
+	TSharedPtr<FNDIStreamReceiver> FindReceiver(const FString& InSourceName);
+	
+	/**
+	 * The given receiver is going to be managed.
+	 * This function may fail if a receiver of the same source is already added.
+	 */
+	bool AddReceiver(const TSharedPtr<FNDIStreamReceiver>& InReceiver);
+
+private:
+	void RemoveExpiredEntries();
+
+private:
+	/** Source name to receiver map. */
+	TArray<TWeakPtr<FNDIStreamReceiver>> ReceiversWeak;
+};
\ No newline at end of file
diff --git a/Engine/Plugins/Media/NDIMedia/Source/NDIMedia/Public/NDIMediaCapture.h b/Engine/Plugins/Media/NDIMedia/Source/NDIMedia/Public/NDIMediaCapture.h
--- a/Engine/Plugins/Media/NDIMedia/Source/NDIMedia/Public/NDIMediaCapture.h
+++ b/Engine/Plugins/Media/NDIMedia/Source/NDIMedia/Public/NDIMediaCapture.h
@@ -7,8 +7,6 @@
 
 #include "NDIMediaCapture.generated.h"
 
-DECLARE_LOG_CATEGORY_EXTERN(LogNDIMedia, Log, All);
-
 /**
  * 
  */
diff --git a/Engine/Plugins/Media/NDIMedia/Source/NDIMedia/Public/NDIMediaDefines.h b/Engine/Plugins/Media/NDIMedia/Source/NDIMedia/Public/NDIMediaDefines.h
new file mode 100644
--- /dev/null
+++ b/Engine/Plugins/Media/NDIMedia/Source/NDIMedia/Public/NDIMediaDefines.h
@@ -0,0 +1,24 @@
+// Copyright Epic Games, Inc. All Rights Reserved.
+
+#pragma once
+
+#include "NDIMediaDefines.generated.h"
+
+/**
+	Receiver Bandwidth modes
+*/
+UENUM(BlueprintType)
+enum class ENDIReceiverBandwidth : uint8
+{
+	/** Receive metadata. */
+	MetadataOnly = 0x00 UMETA(DisplayName = "Metadata Only"),
+
+	/** Receive metadata, audio */
+	AudioOnly = 0x01 UMETA(DisplayName = "Audio Only"),
+
+	/** Receive metadata, audio, video at a lower bandwidth and resolution. */
+	Lowest = 0x02 UMETA(DisplayName = "Lowest"),
+
+	// Receive metadata, audio, video at full resolution.
+	Highest = 0x03 UMETA(DisplayName = "Highest")
+};
diff --git a/Engine/Plugins/Media/NDIMedia/Source/NDIMedia/Public/NDIMediaSource.h b/Engine/Plugins/Media/NDIMedia/Source/NDIMedia/Public/NDIMediaSource.h
new file mode 100644
--- /dev/null
+++ b/Engine/Plugins/Media/NDIMedia/Source/NDIMedia/Public/NDIMediaSource.h
@@ -0,0 +1,96 @@
+// Copyright Epic Games, Inc. All Rights Reserved.
+
+#pragma once
+
+#include "CaptureCardMediaSource.h"
+#include "Engine/TextureDefines.h"
+#include "MediaIOCoreDefinitions.h"
+#include "NDIMediaDefines.h"
+
+#include "NDIMediaSource.generated.h"
+
+/**
+ * Media source for NDI streams.
+ */
+UCLASS(MinimalAPI, BlueprintType, hideCategories=(Platforms,Object), meta=(MediaIOCustomLayout="NDI"))
+class UNDIMediaSource : public UCaptureCardMediaSource
+{
+	GENERATED_BODY()
+
+public:
+	/** Default constructor. */
+	NDIMEDIA_API UNDIMediaSource();
+
+	/** The device, port and video settings that correspond to the input. */
+	UPROPERTY(EditAnywhere, Category="NDI", meta=(DisplayName="Configuration"))
+	FMediaIOConfiguration MediaConfiguration;
+
+	/** Indicates the current bandwidth mode used for the connection to this source */
+	UPROPERTY(EditAnywhere, Category="NDI")
+	ENDIReceiverBandwidth Bandwidth = ENDIReceiverBandwidth::Highest;
+
+	/**
+	 * Indicates whether the timecode should be synced to the Source Timecode value or engine's.
+	 */
+	UPROPERTY(BlueprintReadWrite, EditAnywhere, Category="NDI")
+	bool bSyncTimecodeToSource = true;
+
+	/**
+	 * Capture Ancillary from the NDI source.
+	 * It will decrease performance
+	 */
+	UPROPERTY(BlueprintReadOnly, EditAnywhere, Category="Ancillary")
+	bool bCaptureAncillary;
+
+	/** Maximum number of ancillary data frames to buffer. */
+	UPROPERTY(BlueprintReadOnly, EditAnywhere, AdvancedDisplay, Category="Ancillary", meta=(EditCondition="bCaptureAncillary", ClampMin="1", ClampMax="32"))
+	int32 MaxNumAncillaryFrameBuffer;
+
+	/** Capture Audio from the NDI source. */
+	UPROPERTY(BlueprintReadOnly, EditAnywhere, Category="Audio")
+	bool bCaptureAudio;
+	
+	/** Maximum number of audio frames to buffer. */
+	UPROPERTY(BlueprintReadOnly, EditAnywhere, AdvancedDisplay, Category="Audio", meta=(EditCondition="bCaptureAudio", ClampMin="1", ClampMax="32"))
+	int32 MaxNumAudioFrameBuffer;
+
+	/** Capture Video from the NDI source. */
+	UPROPERTY(BlueprintReadOnly, EditAnywhere, Category="Video")
+	bool bCaptureVideo;
+
+	/** Maximum number of video frames to buffer. */
+	UPROPERTY(BlueprintReadOnly, EditAnywhere, AdvancedDisplay, Category="Video", meta=(EditCondition="bCaptureVideo", ClampMin="1", ClampMax="32"))
+	int32 MaxNumVideoFrameBuffer;
+
+	/** Log a warning when there's a drop frame. */
+	UPROPERTY(EditAnywhere, Category="Debug")
+	bool bLogDropFrame;
+
+	/**
+	 * Burn Frame Timecode in the input texture without any frame number clipping.
+	 * @Note Only supported with progressive format.
+	 */
+	UPROPERTY(BlueprintReadOnly, EditAnywhere, Category="Debug", meta=(DisplayName="Burn Frame Timecode"))
+	bool bEncodeTimecodeInTexel;
+
+	//~ Begin IMediaOptions
+	NDIMEDIA_API virtual bool GetMediaOption(const FName& Key, bool DefaultValue) const override;
+	NDIMEDIA_API virtual int64 GetMediaOption(const FName& Key, int64 DefaultValue) const override;
+	NDIMEDIA_API virtual FString GetMediaOption(const FName& Key, const FString& DefaultValue) const override;
+	NDIMEDIA_API virtual bool HasMediaOption(const FName& Key) const override;
+	//~ End IMediaOptions
+
+	//~ Begin UMediaSource
+	NDIMEDIA_API virtual FString GetUrl() const override;
+	NDIMEDIA_API virtual bool Validate() const override;
+	//~ End UMediaSource
+
+	//~ Begin UObject
+	NDIMEDIA_API virtual void PostLoad() override;
+	NDIMEDIA_API virtual bool SupportsFormatAutoDetection() const override { return true; }
+	//~ End UObject
+
+private:
+	/** Assigns a default configuration if the current configuration is invalid. */
+	void AssignDefaultConfiguration();
+};
diff --git a/Engine/Plugins/Media/NDIMedia/Source/NDIMedia/Public/NDITimecodeProvider.h b/Engine/Plugins/Media/NDIMedia/Source/NDIMedia/Public/NDITimecodeProvider.h
new file mode 100644
--- /dev/null
+++ b/Engine/Plugins/Media/NDIMedia/Source/NDIMedia/Public/NDITimecodeProvider.h
@@ -0,0 +1,63 @@
+// Copyright Epic Games, Inc. All Rights Reserved.
+
+#pragma once
+
+#include "GenlockedTimecodeProvider.h"
+#include "MediaIOCoreDefinitions.h"
+#include "NDIMediaDefines.h"
+
+#include "NDITimecodeProvider.generated.h"
+
+class FNDIStreamReceiver;
+
+/**
+ * Timecode provider from an NDI source 
+ */
+UCLASS(MinimalAPI, Blueprintable, editinlinenew, meta=(DisplayName="NDI Timecode Provider", MediaIOCustomLayout="NDI"))
+class UNDITimecodeProvider : public UGenlockedTimecodeProvider
+{
+	GENERATED_BODY()
+
+public:
+	//~ Begin UTimecodeProvider
+	NDIMEDIA_API virtual bool FetchTimecode(FQualifiedFrameTime& OutFrameTime) override;
+	NDIMEDIA_API virtual ETimecodeProviderSynchronizationState GetSynchronizationState() const override;
+	NDIMEDIA_API virtual bool Initialize(class UEngine* InEngine) override;
+	NDIMEDIA_API virtual void Shutdown(class UEngine* InEngine) override;
+	//~ End UTimecodeProvider
+
+	//~ Begin UObject
+	NDIMEDIA_API virtual void BeginDestroy() override;
+	//~ End UObject
+
+	/**
+	 * Use the time code embedded in the video input stream. 
+	 */
+	UPROPERTY(EditAnywhere, Category="Timecode")
+	FMediaIOVideoTimecodeConfiguration TimecodeConfiguration;
+
+	/** Indicates the current bandwidth mode used for the connection to this source */
+	UPROPERTY(EditAnywhere, Category="NDI")
+	ENDIReceiverBandwidth Bandwidth = ENDIReceiverBandwidth::Highest;
+
+private:
+	void ReleaseResources();
+
+private:
+	/** Critical section protecting "State" */
+	mutable FCriticalSection StateSyncContext;
+
+	/** Current synchronization state */
+	ETimecodeProviderSynchronizationState State = ETimecodeProviderSynchronizationState::Closed;
+
+	/** Last received frame time from the receiver. */
+	FQualifiedFrameTime MostRecentFrameTime;
+
+	/** Current stream receiver */
+	TSharedPtr<FNDIStreamReceiver> Receiver;
+
+	/** Handles for the stream receiver delegates */
+	FDelegateHandle VideoFrameReceivedHandle;
+	FDelegateHandle ConnectedHandle;
+	FDelegateHandle DisconnectedHandle;
+};
diff --git a/Engine/Plugins/Media/NDIMedia/Source/NDIMediaEditor/Private/NDIMediaEditorModule.cpp b/Engine/Plugins/Media/NDIMedia/Source/NDIMediaEditor/Private/NDIMediaEditorModule.cpp
--- a/Engine/Plugins/Media/NDIMedia/Source/NDIMediaEditor/Private/NDIMediaEditorModule.cpp
+++ b/Engine/Plugins/Media/NDIMedia/Source/NDIMediaEditor/Private/NDIMediaEditorModule.cpp
@@ -45,6 +45,8 @@ private:
 
 		StyleInstance->Set("ClassThumbnail.NDIMediaOutput", new IMAGE_BRUSH("NDIMediaOutput_64x", FVector2D(64, 64)));
 		StyleInstance->Set("ClassIcon.NDIMediaOutput", new IMAGE_BRUSH("NDIMediaOutput_20x", FVector2D(20, 20)));
+		StyleInstance->Set("ClassThumbnail.NDIMediaSource", new IMAGE_BRUSH("NDIMediaOutput_64x", FVector2D(64, 64)));
+		StyleInstance->Set("ClassIcon.NDIMediaSource", new IMAGE_BRUSH("NDIMediaOutput_20x", FVector2D(20, 20)));
 
 #undef IMAGE_BRUSH
 
diff --git a/Engine/Plugins/Media/NDIMedia/Source/NDIMediaEditor/Private/NDIMediaSourceFactory.cpp b/Engine/Plugins/Media/NDIMedia/Source/NDIMediaEditor/Private/NDIMediaSourceFactory.cpp
new file mode 100644
--- /dev/null
+++ b/Engine/Plugins/Media/NDIMedia/Source/NDIMediaEditor/Private/NDIMediaSourceFactory.cpp
@@ -0,0 +1,29 @@
+// Copyright Epic Games, Inc. All Rights Reserved.
+
+#include "NDIMediaSourceFactory.h"
+
+#include "AssetTypeCategories.h"
+#include "NDIMediaSource.h"
+
+UNDIMediaSourceFactory::UNDIMediaSourceFactory(const FObjectInitializer& ObjectInitializer)
+	: Super(ObjectInitializer)
+{
+	SupportedClass = UNDIMediaSource::StaticClass();
+	bCreateNew = true;
+	bEditAfterNew = true;
+}
+
+UObject* UNDIMediaSourceFactory::FactoryCreateNew(UClass* InClass, UObject* InParent, FName InName, EObjectFlags InFlags, UObject* InContext, FFeedbackContext* InWarn)
+{
+	return NewObject<UNDIMediaSource>(InParent, InClass, InName, InFlags);
+}
+
+uint32 UNDIMediaSourceFactory::GetMenuCategories() const
+{
+	return EAssetTypeCategories::Media;
+}
+
+bool UNDIMediaSourceFactory::ShouldShowInNewMenu() const
+{
+	return true;
+}
diff --git a/Engine/Plugins/Media/NDIMedia/Source/NDIMediaEditor/Private/NDIMediaSourceFactory.h b/Engine/Plugins/Media/NDIMedia/Source/NDIMediaEditor/Private/NDIMediaSourceFactory.h
new file mode 100644
--- /dev/null
+++ b/Engine/Plugins/Media/NDIMedia/Source/NDIMediaEditor/Private/NDIMediaSourceFactory.h
@@ -0,0 +1,24 @@
+// Copyright Epic Games, Inc. All Rights Reserved.
+
+#pragma once
+
+#include "Factories/Factory.h"
+#include "NDIMediaSourceFactory.generated.h"
+
+/**
+ * Implements a factory for UNDIMediaSource objects.
+ */
+UCLASS()
+class UNDIMediaSourceFactory : public UFactory
+{
+	GENERATED_BODY()
+
+public:
+	UNDIMediaSourceFactory(const FObjectInitializer& ObjectInitializer);
+
+	//~ Begin UFactory
+	virtual UObject* FactoryCreateNew(UClass* InClass, UObject* InParent, FName InName, EObjectFlags InFlags, UObject* InContext, FFeedbackContext* InWarn) override;
+	virtual uint32 GetMenuCategories() const override;
+	virtual bool ShouldShowInNewMenu() const override;
+	//~ End UFactory
+};
