[NDIMedia] Adding support for receiving input with alpha channel.

Adding shaders and custom sample converter for the UYVYA semi planar format (NDI specific).

diff --git a/Engine/Plugins/Media/NDIMedia/NDIMedia.uplugin b/Engine/Plugins/Media/NDIMedia/NDIMedia.uplugin
--- a/Engine/Plugins/Media/NDIMedia/NDIMedia.uplugin
+++ b/Engine/Plugins/Media/NDIMedia/NDIMedia.uplugin
@@ -27,6 +27,12 @@
 			"Type": "Editor",
 			"LoadingPhase": "PostEngineInit",
 			"PlatformAllowList": [ "Win64" ]
+		},
+		{
+			"Name": "NDIMediaRendering",
+			"Type": "Runtime",
+			"LoadingPhase": "PostConfigInit",
+			"PlatformAllowList": [ "Win64" ]
 		}
 	],
 	"Plugins": [
diff --git a/Engine/Plugins/Media/NDIMedia/Shaders/Private/NDIMediaShaders.usf b/Engine/Plugins/Media/NDIMedia/Shaders/Private/NDIMediaShaders.usf
new file mode 100644
--- /dev/null
+++ b/Engine/Plugins/Media/NDIMedia/Shaders/Private/NDIMediaShaders.usf
@@ -0,0 +1,96 @@
+// Copyright Epic Games, Inc. All Rights Reserved.
+
+#include "/Engine/Public/Platform.ush"
+#include "/Engine/Generated/GeneratedUniformBuffers.ush" 
+#include "/Engine/Private/GammaCorrectionCommon.ush"
+#include "/Engine/Private/TonemapCommon.ush"
+
+// ----------------- Begin Duplicated from MediaShaders.usf -------------------------
+// todo: could be extracted in a MediaShadersCommon.ush.
+
+#define UE_Color_EEncoding_Linear		1u
+#define UE_Color_EEncoding_sRGB			2u
+#define UE_Color_EEncoding_ST2084		3u
+#define UE_Color_EEncoding_Gamma22		4u
+#define UE_Color_EEncoding_BT1886		5u
+#define UE_Color_EEncoding_Gamma26		6u
+#define UE_Color_EEncoding_Cineon		7u
+#define UE_Color_EEncoding_REDLog		8u
+#define UE_Color_EEncoding_REDLog3G10	9u
+#define UE_Color_EEncoding_SLog1		10u
+#define UE_Color_EEncoding_SLog2		11u
+#define UE_Color_EEncoding_SLog3		12u
+#define UE_Color_EEncoding_AlexaV3LogC	13u
+#define UE_Color_EEncoding_CanonLog		14u
+#define UE_Color_EEncoding_ProTune		15u
+#define UE_Color_EEncoding_VLog			16u
+
+#define ToneMapMethod_None           0u
+#define ToneMapMethod_Hable          1u
+#define ToneMapMethod_SimpleReinhard 2u
+
+float3 ApplyColorProcessingRGB(float3 Color, float4x4 CSMtx, uint EOTF, uint ToneMapMethod = 0)
+{
+	// Go linear
+	switch(EOTF)
+	{
+	case UE_Color_EEncoding_sRGB:	Color = sRGBToLinear(Color); break;
+	case UE_Color_EEncoding_ST2084: Color = ST2084ToLinear(Color); break;
+	case UE_Color_EEncoding_SLog3:	Color = SLog3ToLinear(Color); break;
+	default: break;
+	}
+
+	// Adjust colorspace
+	float4 Col4 = float4(Color, 1);
+	Color.x = dot(CSMtx[0], Col4);
+	Color.y = dot(CSMtx[1], Col4);
+	Color.z = dot(CSMtx[2], Col4);
+
+	switch(ToneMapMethod)
+	{
+	case ToneMapMethod_Hable:			Color = HableToneMap(Color);  break;
+	case ToneMapMethod_SimpleReinhard:	Color = SimpleReinhardToneMap(Color);  break;
+	default: break;
+	}
+
+	return Color;
+}
+
+float3 ApplyColorProcessing(float3 YUV, float4x4 YUVMtx, float4x4 CSMtx, uint EOTF, uint ToneMapMethod = 0)
+{
+	float3 Color;
+	float4 YUV4 = float4(YUV, 1);
+	Color.x = dot(YUVMtx[0], YUV4);
+	Color.y = dot(YUVMtx[1], YUV4);
+	Color.z = dot(YUVMtx[2], YUV4);
+
+	// Make sure we safely stay above zero
+	Color = max(Color, 0.0f);
+
+	return ApplyColorProcessingRGB(Color, CSMtx, EOTF, ToneMapMethod);
+}
+
+// ----------------- End Duplicated from MediaShaders.usf -------------------------
+
+// Conversion of UYVY-A format to RGBA.
+// Input has 2 textures:
+// - YUVTexture: UYVY 4:2:2
+// - AlphaTexture: A8
+void NDIMediaUYVYAConvertPS(
+	float4 InPosition : SV_POSITION,
+	float2 InUV : TEXCOORD0,
+	out float4 OutColor : SV_Target0)
+{
+	float4 UYVYB = NDIMediaUYVYAConvertUB.YUVTexture.Sample(NDIMediaUYVYAConvertUB.SamplerB, InUV);
+	float4 UYVYT = NDIMediaUYVYAConvertUB.YUVTexture.Sample(NDIMediaUYVYAConvertUB.SamplerT, InUV);
+	float Alpha = NDIMediaUYVYAConvertUB.AlphaTexture.Sample(NDIMediaUYVYAConvertUB.SamplerT, InUV).w;
+
+	float PosX = 2.0f * InUV.x * NDIMediaUYVYAConvertUB.InputWidth;
+
+	float FracX = PosX % 2.0f;
+	float3 YUV;
+	YUV.x = (1 - FracX) * UYVYT.y + FracX * UYVYT.w;
+	YUV.yz = UYVYB.zx;
+
+	OutColor = float4(ApplyColorProcessing(YUV, NDIMediaUYVYAConvertUB.ColorTransform, NDIMediaUYVYAConvertUB.CSTransform, NDIMediaUYVYAConvertUB.EOTF, NDIMediaUYVYAConvertUB.ToneMapMethod), Alpha);
+}
\ No newline at end of file
diff --git a/Engine/Plugins/Media/NDIMedia/Source/NDIMedia/NDIMedia.Build.cs b/Engine/Plugins/Media/NDIMedia/Source/NDIMedia/NDIMedia.Build.cs
--- a/Engine/Plugins/Media/NDIMedia/Source/NDIMedia/NDIMedia.Build.cs
+++ b/Engine/Plugins/Media/NDIMedia/Source/NDIMedia/NDIMedia.Build.cs
@@ -7,7 +7,7 @@ public class NDIMedia : ModuleRules
 	public NDIMedia(ReadOnlyTargetRules Target) : base(Target)
 	{
 		PCHUsage = PCHUsageMode.UseExplicitOrSharedPCHs;
-
+	
 		PublicDependencyModuleNames.AddRange(
 			new string[] {
 				"Core",
@@ -22,6 +22,7 @@ public class NDIMedia : ModuleRules
 				"Engine",
 				"MediaIOCore",
 				"MediaUtils",
+				"NDIMediaRendering",
 				"NDISDK",
 				"Projects",
 				"RenderCore",
diff --git a/Engine/Plugins/Media/NDIMedia/Source/NDIMedia/Private/Player/NDIMediaPlayer.cpp b/Engine/Plugins/Media/NDIMedia/Source/NDIMedia/Private/Player/NDIMediaPlayer.cpp
--- a/Engine/Plugins/Media/NDIMedia/Source/NDIMedia/Private/Player/NDIMediaPlayer.cpp
+++ b/Engine/Plugins/Media/NDIMedia/Source/NDIMedia/Private/Player/NDIMediaPlayer.cpp
@@ -12,97 +12,16 @@
 #include "NDIMediaLog.h"
 #include "NDIMediaModule.h"
 #include "NDIMediaSourceOptions.h"
+#include "NDIMediaTextureSample.h"
+#include "NDIMediaTextureSampleConverter.h"
 #include "NDIStreamReceiver.h"
 #include "NDIStreamReceiverManager.h"
 
 #define LOCTEXT_NAMESPACE "NDIMediaPlayer"
 
-/**
- * Implements a media texture sample for NDIMedia.
- */
-class FNDIMediaTextureSample : public FMediaIOCoreTextureSampleBase
-{
-	using Super = FMediaIOCoreTextureSampleBase;
-
-public:
-	/**
-	 * Initialize the sample.
-	 * 
-	 * @param InVideoFrame Received Video frame data from NDI.
-	 * @param InColorFormatArgs Color space and encoding settings from the media source.
-	 * @param InTime Current playback time of the sample.
-	 * @return true if the sample was correctly initialized, false otherwise.
-	 */
-	bool Initialize(const NDIlib_video_frame_v2_t& InVideoFrame, const UE::MediaIOCore::FColorFormatArgs& InColorFormatArgs, const FTimespan& InTime)
-	{
-		int32 FrameBufferSize;
-		EMediaTextureSampleFormat FrameSampleFormat;
-
-		if (InVideoFrame.FourCC == NDIlib_FourCC_video_type_UYVY)
-		{
-			FrameBufferSize = InVideoFrame.line_stride_in_bytes * InVideoFrame.yres;
-			FrameSampleFormat = EMediaTextureSampleFormat::CharUYVY;
-		}
-		else if (InVideoFrame.FourCC == NDIlib_FourCC_video_type_BGRA)
-		{
-			FrameBufferSize = InVideoFrame.line_stride_in_bytes * InVideoFrame.yres;
-			FrameSampleFormat = EMediaTextureSampleFormat::CharBGRA;
-		}
-		else if (InVideoFrame.FourCC == NDIlib_FourCC_video_type_RGBA
-			|| InVideoFrame.FourCC == NDIlib_FourCC_video_type_RGBX)
-		{
-			FrameBufferSize = InVideoFrame.line_stride_in_bytes * InVideoFrame.yres;
-			FrameSampleFormat = EMediaTextureSampleFormat::CharRGBA;
-		}
-		// else if (InVideoFrame.FourCC == NDIlib_FourCC_video_type_UYVA)
-		// {
-		// 	FrameBufferSize = InVideoFrame.line_stride_in_bytes * InVideoFrame.yres + InVideoFrame.xres*InVideoFrame.yres;
-		// 	FrameSampleFormat = EMediaTextureSampleFormat::CharUYVA; // todo
-		// }
-		else
-		{
-			return false;
-		}
-
-		switch (InVideoFrame.frame_format_type)
-		{
-			case NDIlib_frame_format_type_progressive:
-				return Super::Initialize(InVideoFrame.p_data
-					, FrameBufferSize
-					, InVideoFrame.line_stride_in_bytes
-					, InVideoFrame.xres
-					, InVideoFrame.yres
-					, FrameSampleFormat
-					, InTime
-					, FFrameRate(InVideoFrame.frame_rate_N, InVideoFrame.frame_rate_D)
-					, TOptional<FTimecode>() // todo
-					, InColorFormatArgs);		
-				
-			case NDIlib_frame_format_type_field_0:
-			case NDIlib_frame_format_type_field_1:
-			
-			return Super::InitializeWithEvenOddLine(InVideoFrame.frame_format_type == NDIlib_frame_format_type_field_0
-				, InVideoFrame.p_data
-				, FrameBufferSize
-				, InVideoFrame.line_stride_in_bytes
-				, InVideoFrame.xres
-				, InVideoFrame.yres
-				, FrameSampleFormat
-				, InTime
-				, FFrameRate(InVideoFrame.frame_rate_N, InVideoFrame.frame_rate_D)
-				, TOptional<FTimecode>() // todo
-				, InColorFormatArgs);
-
-			default:
-				return false;
-		}
-	}
-};
-
 class FNDIMediaTextureSamplePool : public TMediaObjectPool<FNDIMediaTextureSample>
 {};
 
-
 /**
  * Implements a media audio sample for NDIMedia
  */
@@ -354,17 +273,27 @@ void FNDIMediaPlayer::HandleVideoFrameReceived(FNDIStreamReceiver* InReceiver, c
 	const UE::MediaIOCore::FColorFormatArgs ColorFormatArgs(
 		bOverrideSourceEncoding ?  static_cast<UE::Color::EEncoding>(OverrideSourceEncoding) : UE::Color::EEncoding::sRGB,
 		bOverrideSourceColorSpace ? static_cast<UE::Color::EColorSpace>(OverrideSourceColorSpace) : UE::Color::EColorSpace::sRGB);
-	
-	if (TextureSample->Initialize(InVideoFrame, ColorFormatArgs, InTime))
+
+	FTimecode SourceTimecode = InReceiver->GetCurrentTimecode();
+
+	if (TextureSample->Initialize(InVideoFrame, ColorFormatArgs, InTime, SourceTimecode))
 	{
+		if (TextureSample->CustomConverter)
+		{
+			TextureSample->CustomConverter->Setup(TextureSample);
+		}
 		TextureSample->SetColorConversionSettings(OCIOSettings);
 
 		if (bEncodeTimecodeInTexel && InVideoFrame.frame_format_type == NDIlib_frame_format_type_progressive)
 		{
 			EMediaIOCoreEncodePixelFormat EncodePixelFormat;
 			bool bEncodeSupported = true;
-			if (InVideoFrame.FourCC == NDIlib_FourCC_video_type_UYVY)
+
+			if (InVideoFrame.FourCC == NDIlib_FourCC_video_type_UYVY
+				|| InVideoFrame.FourCC == NDIlib_FourCC_video_type_UYVA)
 			{
+				// Note: for UYVA, we can write in the UYVY part (even if it ends up being transparent).
+				// todo: add support in FMediaIOCoreEncodeTime for single channel (R) format.
 				EncodePixelFormat = EMediaIOCoreEncodePixelFormat::CharUYVY;
 			}
 			else if (InVideoFrame.FourCC == NDIlib_FourCC_video_type_BGRA
@@ -381,9 +310,8 @@ void FNDIMediaPlayer::HandleVideoFrameReceived(FNDIStreamReceiver* InReceiver, c
 			
 			if (bEncodeSupported)
 			{
-				FTimecode SetTimecode = InReceiver->GetCurrentTimecode();
 				FMediaIOCoreEncodeTime EncodeTime(EncodePixelFormat, const_cast<void*>(TextureSample->GetBuffer()), InVideoFrame.line_stride_in_bytes, InVideoFrame.xres, InVideoFrame.yres);
-				EncodeTime.Render(SetTimecode.Hours, SetTimecode.Minutes, SetTimecode.Seconds, SetTimecode.Frames);
+				EncodeTime.Render(SourceTimecode.Hours, SourceTimecode.Minutes, SourceTimecode.Seconds, SourceTimecode.Frames);
 			}
 		}
 		
diff --git a/Engine/Plugins/Media/NDIMedia/Source/NDIMedia/Private/Player/NDIMediaTextureSample.cpp b/Engine/Plugins/Media/NDIMedia/Source/NDIMedia/Private/Player/NDIMediaTextureSample.cpp
new file mode 100644
--- /dev/null
+++ b/Engine/Plugins/Media/NDIMedia/Source/NDIMedia/Private/Player/NDIMediaTextureSample.cpp
@@ -0,0 +1,95 @@
+// Copyright Epic Games, Inc. All Rights Reserved.
+
+#include "NDIMediaTextureSample.h"
+
+#include "NDIMediaAPI.h"
+#include "NDIMediaTextureSampleConverter.h"
+
+bool FNDIMediaTextureSample::Initialize(const NDIlib_video_frame_v2_t& InVideoFrame, const UE::MediaIOCore::FColorFormatArgs& InColorFormatArgs, const FTimespan& InTime, const TOptional<FTimecode>& InTimecode)
+{
+	bIsCustomFormat = false;
+	bIsProgressive = true;
+	FieldIndex = 0;
+	
+	int32 FrameBufferSize;
+	EMediaTextureSampleFormat FrameSampleFormat;
+
+	if (InVideoFrame.FourCC == NDIlib_FourCC_video_type_UYVY)
+	{
+		FrameBufferSize = InVideoFrame.line_stride_in_bytes * InVideoFrame.yres;
+		FrameSampleFormat = EMediaTextureSampleFormat::CharUYVY;
+	}
+	else if (InVideoFrame.FourCC == NDIlib_FourCC_video_type_BGRA)
+	{
+		FrameBufferSize = InVideoFrame.line_stride_in_bytes * InVideoFrame.yres;
+		FrameSampleFormat = EMediaTextureSampleFormat::CharBGRA;
+	}
+	else if (InVideoFrame.FourCC == NDIlib_FourCC_video_type_RGBA
+		|| InVideoFrame.FourCC == NDIlib_FourCC_video_type_RGBX)
+	{
+		FrameBufferSize = InVideoFrame.line_stride_in_bytes * InVideoFrame.yres;
+		FrameSampleFormat = EMediaTextureSampleFormat::CharRGBA;
+	}
+	else if (InVideoFrame.FourCC == NDIlib_FourCC_video_type_UYVA)
+	{
+		// UYVA format needs a custom converter.
+		bIsCustomFormat = true;
+		FrameBufferSize = InVideoFrame.line_stride_in_bytes * InVideoFrame.yres + InVideoFrame.xres*InVideoFrame.yres;
+		FrameSampleFormat = EMediaTextureSampleFormat::CharRGBA; // Resulting texture needs to be RGBA.
+	}
+	else
+	{
+		return false;
+	}
+
+	// Allocate a custom sample converter if needed.
+	if (!CustomConverter && bIsCustomFormat)
+	{
+		CustomConverter = MakeShared<FNDIMediaTextureSampleConverter>();
+	}
+
+	switch (InVideoFrame.frame_format_type)
+	{
+		case NDIlib_frame_format_type_progressive:
+			bIsProgressive = true;
+			return Super::Initialize(InVideoFrame.p_data
+				, FrameBufferSize
+				, InVideoFrame.line_stride_in_bytes
+				, InVideoFrame.xres
+				, InVideoFrame.yres
+				, FrameSampleFormat
+				, InTime
+				, FFrameRate(InVideoFrame.frame_rate_N, InVideoFrame.frame_rate_D)
+				, InTimecode
+				, InColorFormatArgs);
+
+		case NDIlib_frame_format_type_field_0:
+		case NDIlib_frame_format_type_field_1:
+			bIsProgressive = false;
+			FieldIndex = InVideoFrame.frame_format_type == NDIlib_frame_format_type_field_0 ? 0 : 1;
+		
+		return Super::InitializeWithEvenOddLine(InVideoFrame.frame_format_type == NDIlib_frame_format_type_field_0
+			, InVideoFrame.p_data
+			, FrameBufferSize
+			, InVideoFrame.line_stride_in_bytes
+			, InVideoFrame.xres
+			, InVideoFrame.yres
+			, FrameSampleFormat
+			, InTime
+			, FFrameRate(InVideoFrame.frame_rate_N, InVideoFrame.frame_rate_D)
+			, InTimecode
+			, InColorFormatArgs);
+
+		default:
+			return false;
+	}
+}
+
+IMediaTextureSampleConverter* FNDIMediaTextureSample::GetMediaTextureSampleConverter()
+{
+	if (bIsCustomFormat)
+	{
+		return CustomConverter.Get();
+	}
+	return Super::GetMediaTextureSampleConverter();
+}
diff --git a/Engine/Plugins/Media/NDIMedia/Source/NDIMedia/Private/Player/NDIMediaTextureSample.h b/Engine/Plugins/Media/NDIMedia/Source/NDIMedia/Private/Player/NDIMediaTextureSample.h
new file mode 100644
--- /dev/null
+++ b/Engine/Plugins/Media/NDIMedia/Source/NDIMedia/Private/Player/NDIMediaTextureSample.h
@@ -0,0 +1,45 @@
+// Copyright Epic Games, Inc. All Rights Reserved.
+
+#pragma once
+
+#include "MediaIOCoreTextureSampleBase.h"
+
+class FNDIMediaTextureSampleConverter;
+struct NDIlib_video_frame_v2_t;
+
+/**
+ * Implements a media texture sample for NDIMedia.
+ */
+class FNDIMediaTextureSample : public FMediaIOCoreTextureSampleBase
+{
+	using Super = FMediaIOCoreTextureSampleBase;
+
+public:
+	//~ Begin IMediaTextureSample
+#if WITH_ENGINE
+	virtual IMediaTextureSampleConverter* GetMediaTextureSampleConverter() override;
+#endif
+	//~ End IMediaTextureSample
+
+	/**
+	 * Initialize the sample.
+	 * 
+	 * @param InVideoFrame Received Video frame data from NDI.
+	 * @param InColorFormatArgs Color space and encoding settings from the media source.
+	 * @param InTime Current playback time of the sample.
+	 * @return true if the sample was correctly initialized, false otherwise.
+	 */
+	bool Initialize(const NDIlib_video_frame_v2_t& InVideoFrame, const UE::MediaIOCore::FColorFormatArgs& InColorFormatArgs, const FTimespan& InTime, const TOptional<FTimecode>& InTimecode);
+
+	/** Progressive vs Interlaced. */
+	bool bIsProgressive = true;
+
+	/** If interlaced, which field. */
+	int32 FieldIndex = 0;
+
+	/** Needs a custom conversion. */
+	bool bIsCustomFormat = false;
+
+	/** Custom converter. */
+	TSharedPtr<FNDIMediaTextureSampleConverter> CustomConverter;
+};
diff --git a/Engine/Plugins/Media/NDIMedia/Source/NDIMedia/Private/Player/NDIMediaTextureSampleConverter.cpp b/Engine/Plugins/Media/NDIMedia/Source/NDIMedia/Private/Player/NDIMediaTextureSampleConverter.cpp
new file mode 100644
--- /dev/null
+++ b/Engine/Plugins/Media/NDIMedia/Source/NDIMedia/Private/Player/NDIMediaTextureSampleConverter.cpp
@@ -0,0 +1,179 @@
+// Copyright Epic Games, Inc. All Rights Reserved.
+
+#include "NDIMediaTextureSampleConverter.h"
+
+#include "NDIMediaShaders.h"
+#include "NDIMediaTextureSample.h"
+#include "PipelineStateCache.h"
+#include "RHIStaticStates.h"
+
+namespace UE::NDIMediaTextureSampleConverter
+{
+	// Reference FMediaTextureResource::GetColorSpaceConversionMatrixForSample.
+	FMatrix44f GetColorSpaceConversionMatrixForSample(const TSharedPtr<FNDIMediaTextureSample>& InSample)
+	{
+		FMatrix44f OutColorSpaceMatrix;
+		
+		const UE::Color::FColorSpace& Working = UE::Color::FColorSpace::GetWorking();
+
+		if (InSample->GetMediaTextureSampleColorConverter())
+		{
+			OutColorSpaceMatrix = FMatrix44f::Identity;
+		}
+		else
+		{
+			OutColorSpaceMatrix = UE::Color::Transpose<float>(UE::Color::FColorSpaceTransform(InSample->GetSourceColorSpace(), Working));
+		}
+	
+		const float NormalizationFactor = InSample->GetHDRNitsNormalizationFactor();
+		if (NormalizationFactor != 1.0f)
+		{
+			OutColorSpaceMatrix = OutColorSpaceMatrix.ApplyScale(NormalizationFactor);
+		}
+
+		return OutColorSpaceMatrix;
+	}
+}
+
+void FNDIMediaTextureSampleConverter::Setup(const TSharedPtr<FNDIMediaTextureSample>& InSample)
+{
+	SampleWeak = InSample;
+}
+
+uint32 FNDIMediaTextureSampleConverter::GetConverterInfoFlags() const
+{
+	return ConverterInfoFlags_Default;
+}
+
+// Reference: FMediaTextureResource::ConvertTextureToOutput
+bool FNDIMediaTextureSampleConverter::Convert(FRHICommandListImmediate& InRHICmdList, FTextureRHIRef& InDstTexture, const FConversionHints& InHints)
+{
+	TSharedPtr<FNDIMediaTextureSample> Sample = SampleWeak.Pin();
+	if (!Sample || !Sample->GetBuffer())
+	{
+		return false;
+	}
+
+	if (!UpdateInputTextures(InRHICmdList, Sample))
+	{
+		return false;
+	}
+	
+	// Initialize the frame size parameter
+	FIntPoint FieldSize = Sample->GetDim();
+	FIntPoint FrameSize = Sample->bIsProgressive ? FieldSize : FIntPoint(FieldSize.X, FieldSize.Y*2);
+	
+	// Draw full size quad into render target
+	// This needs to happen before we begin to setup the draw call, because on DX11, this might flush the command list more or less randomly
+	const float FieldUVOffset = Sample->FieldIndex ? 0.5f/FrameSize.Y : 0.f;
+	FBufferRHIRef VertexBuffer = CreateTempMediaVertexBuffer(InRHICmdList, 0.f, 1.f, 0.f-FieldUVOffset, 1.f-FieldUVOffset);
+
+	FGraphicsPipelineStateInitializer GraphicsPSOInit;
+
+	InRHICmdList.Transition(FRHITransitionInfo(InDstTexture, ERHIAccess::Unknown, ERHIAccess::RTV));
+	ON_SCOPE_EXIT {InRHICmdList.Transition(FRHITransitionInfo(InDstTexture, ERHIAccess::RTV, ERHIAccess::SRVGraphics));};
+	
+	FRHIRenderPassInfo RPInfo(InDstTexture.GetReference(), ERenderTargetActions::DontLoad_Store);
+	{
+		InRHICmdList.BeginRenderPass(RPInfo, TEXT("ConvertMedia(NDI)"));
+		ON_SCOPE_EXIT { InRHICmdList.EndRenderPass(); };
+		
+		InRHICmdList.ApplyCachedRenderTargets(GraphicsPSOInit);
+		
+		FIntPoint OutputSize(InDstTexture->GetSizeXYZ().X, InDstTexture->GetSizeXYZ().Y);
+		InRHICmdList.SetViewport(0, 0, 0.0f, OutputSize.X, OutputSize.Y, 1.0f);
+	
+		GraphicsPSOInit.DepthStencilState = TStaticDepthStencilState<false, CF_Always>::GetRHI();
+		GraphicsPSOInit.RasterizerState = TStaticRasterizerState<>::GetRHI();
+		GraphicsPSOInit.BlendState = TStaticBlendStateWriteMask<CW_RGBA, CW_NONE, CW_NONE, CW_NONE, CW_NONE, CW_NONE, CW_NONE, CW_NONE>::GetRHI();
+		GraphicsPSOInit.PrimitiveType = PT_TriangleStrip;
+
+		// Configure media shaders
+		FGlobalShaderMap* ShaderMap = GetGlobalShaderMap(GMaxRHIFeatureLevel);
+		TShaderMapRef<FMediaShadersVS> VertexShader(ShaderMap);
+		GraphicsPSOInit.BoundShaderState.VertexDeclarationRHI = GMediaVertexDeclaration.VertexDeclarationRHI;
+		GraphicsPSOInit.BoundShaderState.VertexShaderRHI = VertexShader.GetVertexShader();
+
+		// Note: this convert supports only one shader, all others are already supported by media shaders.
+		{	// --- Setup UYVA to BGRA Pixel shader ---
+			TShaderMapRef<FNDIMediaShaderUYVAtoBGRAPS> ConvertShader(ShaderMap);
+			GraphicsPSOInit.BoundShaderState.PixelShaderRHI = ConvertShader.GetPixelShader();
+
+			// Ensure the pipeline state is set to the one we've configured
+			SetGraphicsPipelineState(InRHICmdList, GraphicsPSOInit, 0);
+	
+			// set the texture parameter of the conversion shader
+			FNDIMediaShaderUYVAtoBGRAPS::FParameters Params(
+				SourceYUVTexture,
+				SourceAlphaTexture,
+				OutputSize,
+				Sample->GetSampleToRGBMatrix(),
+				Sample->GetEncodingType(),
+				UE::NDIMediaTextureSampleConverter::GetColorSpaceConversionMatrixForSample(Sample),
+				Sample->GetToneMapMethod());
+
+			SetShaderParametersLegacyPS(InRHICmdList, ConvertShader, Params);
+		}
+	
+		InRHICmdList.SetStreamSource(0, VertexBuffer, 0);
+
+		// Set Viewport to RT size
+		InRHICmdList.SetViewport(0, 0, 0.0f, OutputSize.X, OutputSize.Y, 1.0f);
+
+		InRHICmdList.DrawPrimitive(0, 2, 1);
+	}
+	return true;
+}
+
+bool FNDIMediaTextureSampleConverter::UpdateInputTextures(FRHICommandList& InRHICmdList, const TSharedPtr<FNDIMediaTextureSample>& InSample)
+{
+	const void* Buffer = InSample->GetBuffer();
+	
+	if (!Buffer)
+	{
+		return false;
+	}
+	
+	FIntPoint FieldSize = InSample->GetDim();
+	FIntPoint FrameSize = InSample->bIsProgressive ? FieldSize : FIntPoint(FieldSize.X, FieldSize.Y*2);
+
+	if (!SourceYUVTexture.IsValid() || FrameSize != PreviousFrameSize)
+	{
+		// The source YUV texture will be given UYVY data, so make it half-width
+		const FRHITextureCreateDesc CreateDesc = FRHITextureCreateDesc::Create2D(TEXT("NDIMediaReceiverInterlacedAlphaSourceTexture"))
+			.SetExtent(FieldSize.X / 2, FieldSize.Y)
+			.SetFormat(PF_B8G8R8A8)
+			.SetNumMips(1)
+			.SetFlags(ETextureCreateFlags::RenderTargetable | ETextureCreateFlags::Dynamic);
+
+		SourceYUVTexture = RHICreateTexture(CreateDesc);
+	}
+
+	if (!SourceAlphaTexture.IsValid() || FrameSize != PreviousFrameSize)
+	{
+		const FRHITextureCreateDesc CreateAlphaDesc = FRHITextureCreateDesc::Create2D(TEXT("NDIMediaReceiverInterlacedAlphaSourceAlphaTexture"))
+			.SetExtent(FieldSize.X, FieldSize.Y)
+			.SetFormat(PF_A8)
+			.SetNumMips(1)
+			.SetFlags(ETextureCreateFlags::RenderTargetable | ETextureCreateFlags::Dynamic);
+
+		SourceAlphaTexture = RHICreateTexture(CreateAlphaDesc);
+	}
+
+	PreviousFrameSize = FrameSize;
+
+	// Update GPU textures from sample buffer
+	FUpdateTextureRegion2D YUVRegion(0, 0, 0, 0, FieldSize.X/2, FieldSize.Y);
+	RHIUpdateTexture2D(SourceYUVTexture, 0, YUVRegion, InSample->GetStride(), (uint8*&)Buffer);
+
+	// Make sure resource is in SRV mode again
+	InRHICmdList.Transition(FRHITransitionInfo(SourceYUVTexture, ERHIAccess::Unknown, ERHIAccess::SRVMask));
+	
+	FUpdateTextureRegion2D AlphaRegion(0, 0, 0, 0, FieldSize.X, FieldSize.Y);
+	RHIUpdateTexture2D(SourceAlphaTexture, 0, AlphaRegion, FrameSize.X, ((uint8*&)Buffer) + FieldSize.Y * InSample->GetStride());
+
+	// Make sure resource is in SRV mode again
+	InRHICmdList.Transition(FRHITransitionInfo(SourceAlphaTexture, ERHIAccess::Unknown, ERHIAccess::SRVMask));
+
+	return true;
+}
diff --git a/Engine/Plugins/Media/NDIMedia/Source/NDIMedia/Private/Player/NDIMediaTextureSampleConverter.h b/Engine/Plugins/Media/NDIMedia/Source/NDIMedia/Private/Player/NDIMediaTextureSampleConverter.h
new file mode 100644
--- /dev/null
+++ b/Engine/Plugins/Media/NDIMedia/Source/NDIMedia/Private/Player/NDIMediaTextureSampleConverter.h
@@ -0,0 +1,35 @@
+// Copyright Epic Games, Inc. All Rights Reserved.
+
+#pragma once
+
+#include "IMediaTextureSampleConverter.h"
+
+class FNDIMediaTextureSample;
+
+class FNDIMediaTextureSampleConverter : public IMediaTextureSampleConverter
+{
+public:
+	/** Configures settings to convert incoming sample */
+	void Setup(const TSharedPtr<FNDIMediaTextureSample>& InSample);
+
+	//~ Begin IMediaTextureSampleConverter
+	virtual uint32 GetConverterInfoFlags() const override;
+	virtual bool Convert(FRHICommandListImmediate& InRHICmdList, FTextureRHIRef& InDstTexture, const FConversionHints& InHints) override;
+	//~ End IMediaTextureSampleConverter
+	
+private:
+	/** Prepare the input textures from the sample buffer. */
+	bool UpdateInputTextures(FRHICommandList& InRHICmdList, const TSharedPtr<FNDIMediaTextureSample>& InSample);
+	
+	/** Keep a reference to the sample to retrieve frame info and buffer. */
+	TWeakPtr<FNDIMediaTextureSample> SampleWeak;
+
+	/** Cache the last frame size of the texture objects. */
+	FIntPoint PreviousFrameSize = FIntPoint(0, 0);
+
+	/** Source YUV texture */
+	FTextureRHIRef SourceYUVTexture;
+
+	/** Source Alpha texture */
+	FTextureRHIRef SourceAlphaTexture;
+};
diff --git a/Engine/Plugins/Media/NDIMedia/Source/NDIMediaRendering/Internal/NDIMediaShaders.h b/Engine/Plugins/Media/NDIMedia/Source/NDIMediaRendering/Internal/NDIMediaShaders.h
new file mode 100644
--- /dev/null
+++ b/Engine/Plugins/Media/NDIMedia/Source/NDIMediaRendering/Internal/NDIMediaShaders.h
@@ -0,0 +1,65 @@
+// Copyright Epic Games, Inc. All Rights Reserved.
+
+#pragma once
+
+#include "ColorManagement/ColorManagementDefines.h"
+#include "DataDrivenShaderPlatformInfo.h"
+#include "GlobalShader.h"
+#include "MediaShaders.h"
+#include "RHI.h"
+#include "Shader.h"
+#include "ShaderParameterUtils.h"
+
+/**
+ * Pixel shader for ndi media sample converter.
+ * Converts UYVY-A 2 planes format.
+ * The first plane is a regular UYVY 4:2:2 followed by a alpha channel plane.  
+ */
+class FNDIMediaShaderUYVAtoBGRAPS : public FGlobalShader
+{
+	DECLARE_EXPORTED_GLOBAL_SHADER(FNDIMediaShaderUYVAtoBGRAPS, NDIMEDIARENDERING_API);
+
+public:
+	static bool ShouldCompilePermutation(const FGlobalShaderPermutationParameters& InParameters)
+	{
+		return IsFeatureLevelSupported(InParameters.Platform, ERHIFeatureLevel::ES3_1);
+	}
+
+	FNDIMediaShaderUYVAtoBGRAPS()
+	{}
+
+	FNDIMediaShaderUYVAtoBGRAPS(const ShaderMetaType::CompiledShaderInitializerType& InInitializer)
+		: FGlobalShader(InInitializer)
+	{}
+
+	struct FParameters
+	{
+		FParameters(
+			const TRefCountPtr<FRHITexture>& InYUVTexture,
+			const TRefCountPtr<FRHITexture>& InAlphaTexture,
+			const FIntPoint& InOutputSize,
+			const FMatrix44f& InColorTransform,
+			UE::Color::EEncoding InEncoding,
+			const FMatrix44f& InCSTransform,
+			MediaShaders::EToneMapMethod InToneMapMethod)
+			: YUVTexture(InYUVTexture)
+			, AlphaTexture(InAlphaTexture)
+			, OutputSize(InOutputSize)
+			, ColorTransform(InColorTransform)
+			, Encoding(InEncoding)
+			, CSTransform(InCSTransform)
+			, ToneMapMethod(InToneMapMethod)
+		{
+		}
+		
+		TRefCountPtr<FRHITexture> YUVTexture;
+		TRefCountPtr<FRHITexture> AlphaTexture;
+		FIntPoint OutputSize;
+		FMatrix44f ColorTransform;
+		UE::Color::EEncoding Encoding;
+		FMatrix44f CSTransform;
+		MediaShaders::EToneMapMethod ToneMapMethod;
+	};
+
+	NDIMEDIARENDERING_API void SetParameters(FRHIBatchedShaderParameters& InBatchedParameters, const FParameters& InParameters);
+};
\ No newline at end of file
diff --git a/Engine/Plugins/Media/NDIMedia/Source/NDIMediaRendering/NDIMediaRendering.Build.cs b/Engine/Plugins/Media/NDIMedia/Source/NDIMediaRendering/NDIMediaRendering.Build.cs
new file mode 100644
--- /dev/null
+++ b/Engine/Plugins/Media/NDIMedia/Source/NDIMediaRendering/NDIMediaRendering.Build.cs
@@ -0,0 +1,23 @@
+// Copyright Epic Games, Inc. All Rights Reserved.
+
+using UnrealBuildTool;
+
+public class NDIMediaRendering : ModuleRules
+{
+	public NDIMediaRendering(ReadOnlyTargetRules Target) : base(Target)
+	{
+		PCHUsage = PCHUsageMode.UseExplicitOrSharedPCHs;
+
+		PublicDependencyModuleNames.AddRange(
+			new string[] {
+			});
+
+		PrivateDependencyModuleNames.AddRange(
+			new string[] {
+				"Core",
+				"Projects",
+				"RenderCore",
+				"RHI"
+			});
+	}
+}
diff --git a/Engine/Plugins/Media/NDIMedia/Source/NDIMediaRendering/Private/NDIMediaRenderingModule.cpp b/Engine/Plugins/Media/NDIMedia/Source/NDIMediaRendering/Private/NDIMediaRenderingModule.cpp
new file mode 100644
--- /dev/null
+++ b/Engine/Plugins/Media/NDIMedia/Source/NDIMediaRendering/Private/NDIMediaRenderingModule.cpp
@@ -0,0 +1,16 @@
+// Copyright Epic Games, Inc. All Rights Reserved.
+
+#include "NDIMediaRenderingModule.h"
+
+#include "Interfaces/IPluginManager.h"
+#include "Misc/Paths.h"
+#include "Modules/ModuleManager.h"
+#include "ShaderCore.h"
+
+void FNDIMediaRenderingModule::StartupModule()
+{
+	const FString PluginShaderDir = FPaths::Combine(IPluginManager::Get().FindPlugin(TEXT("NDIMedia"))->GetBaseDir(), TEXT("Shaders"));
+	AddShaderSourceDirectoryMapping(TEXT("/Plugin/NDIMedia"), PluginShaderDir);
+}
+
+IMPLEMENT_MODULE(FNDIMediaRenderingModule, NDIMediaRendering);
diff --git a/Engine/Plugins/Media/NDIMedia/Source/NDIMediaRendering/Private/NDIMediaRenderingModule.h b/Engine/Plugins/Media/NDIMedia/Source/NDIMediaRendering/Private/NDIMediaRenderingModule.h
new file mode 100644
--- /dev/null
+++ b/Engine/Plugins/Media/NDIMedia/Source/NDIMediaRendering/Private/NDIMediaRenderingModule.h
@@ -0,0 +1,16 @@
+// Copyright Epic Games, Inc. All Rights Reserved.
+
+#pragma once
+
+#include "Modules/ModuleInterface.h"
+
+/**
+ * Module used to define NDI media shaders used for pixel format conversion
+ */
+class FNDIMediaRenderingModule : public IModuleInterface
+{
+public:
+	//~ Begin IModuleInterface
+	virtual void StartupModule() override;
+	//~ End IModuleInterface
+};
diff --git a/Engine/Plugins/Media/NDIMedia/Source/NDIMediaRendering/Private/NDIMediaShaders.cpp b/Engine/Plugins/Media/NDIMedia/Source/NDIMediaRendering/Private/NDIMediaShaders.cpp
new file mode 100644
--- /dev/null
+++ b/Engine/Plugins/Media/NDIMedia/Source/NDIMediaRendering/Private/NDIMediaShaders.cpp
@@ -0,0 +1,46 @@
+// Copyright Epic Games, Inc. All Rights Reserved.
+
+#include "NDIMediaShaders.h"
+
+#include "RHIStaticStates.h"
+
+BEGIN_GLOBAL_SHADER_PARAMETER_STRUCT(FNDIMediaUYVYAConvertUB, )
+	SHADER_PARAMETER(uint32, InputWidth)
+	SHADER_PARAMETER(uint32, InputHeight)
+	SHADER_PARAMETER(uint32, OutputWidth)
+	SHADER_PARAMETER(uint32, OutputHeight)
+	SHADER_PARAMETER(FMatrix44f, ColorTransform)
+	SHADER_PARAMETER(FMatrix44f, CSTransform)
+	SHADER_PARAMETER(uint32, EOTF)
+	SHADER_PARAMETER(uint32, ToneMapMethod)
+	SHADER_PARAMETER_TEXTURE(Texture2D, YUVTexture)
+	SHADER_PARAMETER_TEXTURE(Texture2D, AlphaTexture)
+	SHADER_PARAMETER_SAMPLER(SamplerState, SamplerB)
+	SHADER_PARAMETER_SAMPLER(SamplerState, SamplerT)
+END_GLOBAL_SHADER_PARAMETER_STRUCT()
+
+IMPLEMENT_GLOBAL_SHADER_PARAMETER_STRUCT(FNDIMediaUYVYAConvertUB, "NDIMediaUYVYAConvertUB");
+
+IMPLEMENT_GLOBAL_SHADER(FNDIMediaShaderUYVAtoBGRAPS, "/Plugin/NDIMedia/Private/NDIMediaShaders.usf", "NDIMediaUYVYAConvertPS", SF_Pixel);
+
+void FNDIMediaShaderUYVAtoBGRAPS::SetParameters(FRHIBatchedShaderParameters& InBatchedParameters, const FParameters& InParameters)
+{
+	FNDIMediaUYVYAConvertUB UB;
+	{
+		UB.InputWidth = InParameters.YUVTexture->GetSizeX();
+		UB.InputHeight = InParameters.YUVTexture->GetSizeY();
+		UB.OutputWidth = InParameters.OutputSize.X;
+		UB.OutputHeight = InParameters.OutputSize.Y;
+		UB.ColorTransform = InParameters.ColorTransform;
+		UB.CSTransform = InParameters.CSTransform;
+		UB.EOTF = static_cast<uint32>(InParameters.Encoding);
+		UB.ToneMapMethod = static_cast<uint32>(InParameters.ToneMapMethod);
+		UB.YUVTexture = InParameters.YUVTexture;
+		UB.AlphaTexture = InParameters.AlphaTexture;
+		UB.SamplerB = TStaticSamplerState<SF_Bilinear>::GetRHI();
+		UB.SamplerT = TStaticSamplerState<SF_Trilinear>::GetRHI();
+	}
+
+	TUniformBufferRef<FNDIMediaUYVYAConvertUB> Data = TUniformBufferRef<FNDIMediaUYVYAConvertUB>::CreateUniformBufferImmediate(UB, UniformBuffer_SingleFrame);
+	SetUniformBufferParameter(InBatchedParameters, GetUniformBufferParameter<FNDIMediaUYVYAConvertUB>(), Data);
+}
\ No newline at end of file
