[NDIMedia] Fix for content browser issues by adding asset definitions for NDIMediaSource and NDIMediaOutput

Fix for:
- Asset name will disappear from Content Browser if asset thumbnail are too small
- Asset classes don't show up in filter

diff --git a/Engine/Plugins/Media/NDIMedia/Source/NDIMediaEditor/NDIMediaEditor.Build.cs b/Engine/Plugins/Media/NDIMedia/Source/NDIMediaEditor/NDIMediaEditor.Build.cs
--- a/Engine/Plugins/Media/NDIMedia/Source/NDIMediaEditor/NDIMediaEditor.Build.cs
+++ b/Engine/Plugins/Media/NDIMedia/Source/NDIMediaEditor/NDIMediaEditor.Build.cs
@@ -30,6 +30,7 @@ public class NDIMediaEditor : ModuleRules
 
 		PrivateDependencyModuleNames.AddRange(
 		new string[] {
+				"AssetDefinition",
 				"Projects",
 			});
 	}
diff --git a/Engine/Plugins/Media/NDIMedia/Source/NDIMediaEditor/Private/AssetDefinition_NDIMediaOutput.cpp b/Engine/Plugins/Media/NDIMedia/Source/NDIMediaEditor/Private/AssetDefinition_NDIMediaOutput.cpp
new file mode 100644
--- /dev/null
+++ b/Engine/Plugins/Media/NDIMedia/Source/NDIMediaEditor/Private/AssetDefinition_NDIMediaOutput.cpp
@@ -0,0 +1,32 @@
+// Copyright Epic Games, Inc. All Rights Reserved.
+
+#include "AssetDefinition_NDIMediaOutput.h"
+
+#include "NDIMediaOutput.h"
+#include "ToolMenus.h"
+
+#define LOCTEXT_NAMESPACE "AssetDefinition_NDIMediaOutput"
+
+FText UAssetDefinition_NDIMediaOutput::GetAssetDisplayName() const
+{
+	return LOCTEXT("AssetDisplayName", "NDI Media Output");
+}
+
+FLinearColor UAssetDefinition_NDIMediaOutput::GetAssetColor() const
+{
+	// From NDI Brand Guidelines:  #6257FF
+	return FLinearColor(FColor(98, 87, 255));
+}
+
+TSoftClassPtr<UObject> UAssetDefinition_NDIMediaOutput::GetAssetClass() const
+{
+	return UNDIMediaOutput::StaticClass();
+}
+
+TConstArrayView<FAssetCategoryPath> UAssetDefinition_NDIMediaOutput::GetAssetCategories() const
+{
+	static const TArray<FAssetCategoryPath> Categories = {EAssetCategoryPaths::Media};
+	return Categories;
+}
+
+#undef LOCTEXT_NAMESPACE
diff --git a/Engine/Plugins/Media/NDIMedia/Source/NDIMediaEditor/Private/AssetDefinition_NDIMediaOutput.h b/Engine/Plugins/Media/NDIMedia/Source/NDIMediaEditor/Private/AssetDefinition_NDIMediaOutput.h
new file mode 100644
--- /dev/null
+++ b/Engine/Plugins/Media/NDIMedia/Source/NDIMediaEditor/Private/AssetDefinition_NDIMediaOutput.h
@@ -0,0 +1,20 @@
+// Copyright Epic Games, Inc. All Rights Reserved.
+
+#pragma once
+
+#include "AssetDefinitionDefault.h"
+#include "AssetDefinition_NDIMediaOutput.generated.h"
+
+UCLASS()
+class UAssetDefinition_NDIMediaOutput : public UAssetDefinitionDefault
+{
+	GENERATED_BODY()
+
+public:
+	//~ Begin UAssetDefinition
+	virtual FText GetAssetDisplayName() const override;
+	virtual FLinearColor GetAssetColor() const override;
+	virtual TSoftClassPtr<UObject> GetAssetClass() const override;
+	virtual TConstArrayView<FAssetCategoryPath> GetAssetCategories() const override;
+	//~ End UAssetDefinition
+};
diff --git a/Engine/Plugins/Media/NDIMedia/Source/NDIMediaEditor/Private/AssetDefinition_NDIMediaSource.cpp b/Engine/Plugins/Media/NDIMedia/Source/NDIMediaEditor/Private/AssetDefinition_NDIMediaSource.cpp
new file mode 100644
--- /dev/null
+++ b/Engine/Plugins/Media/NDIMedia/Source/NDIMediaEditor/Private/AssetDefinition_NDIMediaSource.cpp
@@ -0,0 +1,32 @@
+// Copyright Epic Games, Inc. All Rights Reserved.
+
+#include "AssetDefinition_NDIMediaSource.h"
+
+#include "NDIMediaSource.h"
+#include "ToolMenus.h"
+
+#define LOCTEXT_NAMESPACE "AssetDefinition_NDIMediaSource"
+
+FText UAssetDefinition_NDIMediaSource::GetAssetDisplayName() const
+{
+	return LOCTEXT("AssetDisplayName", "NDI Media Source");
+}
+
+FLinearColor UAssetDefinition_NDIMediaSource::GetAssetColor() const
+{
+	// From NDI Brand Guidelines:  #6257FF
+	return FLinearColor(FColor(98, 87, 255));
+}
+
+TSoftClassPtr<UObject> UAssetDefinition_NDIMediaSource::GetAssetClass() const
+{
+	return UNDIMediaSource::StaticClass();
+}
+
+TConstArrayView<FAssetCategoryPath> UAssetDefinition_NDIMediaSource::GetAssetCategories() const
+{
+	static const TArray<FAssetCategoryPath> Categories = {EAssetCategoryPaths::Media};
+	return Categories;
+}
+
+#undef LOCTEXT_NAMESPACE
diff --git a/Engine/Plugins/Media/NDIMedia/Source/NDIMediaEditor/Private/AssetDefinition_NDIMediaSource.h b/Engine/Plugins/Media/NDIMedia/Source/NDIMediaEditor/Private/AssetDefinition_NDIMediaSource.h
new file mode 100644
--- /dev/null
+++ b/Engine/Plugins/Media/NDIMedia/Source/NDIMediaEditor/Private/AssetDefinition_NDIMediaSource.h
@@ -0,0 +1,20 @@
+// Copyright Epic Games, Inc. All Rights Reserved.
+
+#pragma once
+
+#include "AssetDefinitionDefault.h"
+#include "AssetDefinition_NDIMediaSource.generated.h"
+
+UCLASS()
+class UAssetDefinition_NDIMediaSource : public UAssetDefinitionDefault
+{
+	GENERATED_BODY()
+
+public:
+	//~ Begin UAssetDefinition
+	virtual FText GetAssetDisplayName() const override;
+	virtual FLinearColor GetAssetColor() const override;
+	virtual TSoftClassPtr<UObject> GetAssetClass() const override;
+	virtual TConstArrayView<FAssetCategoryPath> GetAssetCategories() const override;
+	//~ End UAssetDefinition
+};
