[NDIMedia] Fixing NDIMediaSource editor

There is already a class hierarchy for MediaSourceActions that need to be used to inherit the open editor action. We can't use the new asset definition.

diff --git a/Engine/Plugins/Media/NDIMedia/Source/NDIMediaEditor/NDIMediaEditor.Build.cs b/Engine/Plugins/Media/NDIMedia/Source/NDIMediaEditor/NDIMediaEditor.Build.cs
--- a/Engine/Plugins/Media/NDIMedia/Source/NDIMediaEditor/NDIMediaEditor.Build.cs
+++ b/Engine/Plugins/Media/NDIMedia/Source/NDIMediaEditor/NDIMediaEditor.Build.cs
@@ -31,6 +31,7 @@ public class NDIMediaEditor : ModuleRules
 		PrivateDependencyModuleNames.AddRange(
 		new string[] {
 				"AssetDefinition",
+				"MediaPlayerEditor",
 				"Projects",
 			});
 	}
diff --git a/Engine/Plugins/Media/NDIMedia/Source/NDIMediaEditor/Private/AssetDefinition_NDIMediaSource.cpp b/Engine/Plugins/Media/NDIMedia/Source/NDIMediaEditor/Private/AssetDefinition_NDIMediaSource.cpp
deleted file mode 100644
diff --git a/Engine/Plugins/Media/NDIMedia/Source/NDIMediaEditor/Private/AssetDefinition_NDIMediaSource.h b/Engine/Plugins/Media/NDIMedia/Source/NDIMediaEditor/Private/AssetDefinition_NDIMediaSource.h
deleted file mode 100644
diff --git a/Engine/Plugins/Media/NDIMedia/Source/NDIMediaEditor/Private/NDIMediaEditorModule.cpp b/Engine/Plugins/Media/NDIMedia/Source/NDIMediaEditor/Private/NDIMediaEditorModule.cpp
--- a/Engine/Plugins/Media/NDIMedia/Source/NDIMediaEditor/Private/NDIMediaEditorModule.cpp
+++ b/Engine/Plugins/Media/NDIMedia/Source/NDIMediaEditor/Private/NDIMediaEditorModule.cpp
@@ -2,11 +2,13 @@
 
 #include "NDIMediaEditorModule.h"
 
+#include "AssetToolsModule.h"
 #include "Brushes/SlateImageBrush.h"
 #include "Interfaces/IPluginManager.h"
 #include "Misc/Paths.h"
 #include "Modules/ModuleInterface.h"
 #include "Modules/ModuleManager.h"
+#include "NDIMediaSourceActions.h"
 #include "PropertyEditorDelegates.h"
 #include "PropertyEditorModule.h"
 #include "Styling/SlateStyle.h"
@@ -23,6 +25,7 @@ class FNDIMediaEditorModule : public IModuleInterface
 	virtual void StartupModule() override
 	{
 		RegisterStyle();
+		RegisterAssetTools();
 	}
 
 	virtual void ShutdownModule() override
@@ -30,10 +33,33 @@ class FNDIMediaEditorModule : public IModuleInterface
 		if (!UObjectInitialized() && !IsEngineExitRequested())
 		{
 			UnregisterStyle();
+			UnregisterAssetTools();
 		}
 	}
 
 private:
+	void RegisterAssetTools()
+	{
+		IAssetTools& AssetTools = FModuleManager::LoadModuleChecked<FAssetToolsModule>("AssetTools").Get();
+
+		TSharedRef<IAssetTypeActions> Action = MakeShared<FNDIMediaSourceActions>();
+		AssetTools.RegisterAssetTypeActions(Action);
+		RegisteredAssetTypeActions.Add(Action);
+	}
+
+	void UnregisterAssetTools()
+	{
+		if (FAssetToolsModule* AssetToolsModule = FModuleManager::GetModulePtr<FAssetToolsModule>("AssetTools"))
+		{
+			IAssetTools& AssetTools = AssetToolsModule->Get();
+			for (const TSharedRef<IAssetTypeActions>& Action : RegisteredAssetTypeActions)
+			{
+				AssetTools.UnregisterAssetTypeActions(Action);
+			}
+			RegisteredAssetTypeActions.Reset();
+		}
+	}
+
 	/** Register slate style */
 	void RegisterStyle()
 	{
@@ -58,6 +84,10 @@ private:
 	{
 		FSlateStyleRegistry::UnRegisterSlateStyle(*StyleInstance.Get());
 	}
+	
+private:
+	/** The collection of registered asset type actions. */
+	TArray<TSharedRef<IAssetTypeActions>> RegisteredAssetTypeActions;
 };
 
 #undef LOCTEXT_NAMESPACE
diff --git a/Engine/Plugins/Media/NDIMedia/Source/NDIMediaEditor/Private/NDIMediaSourceActions.cpp b/Engine/Plugins/Media/NDIMedia/Source/NDIMediaEditor/Private/NDIMediaSourceActions.cpp
new file mode 100644
--- /dev/null
+++ b/Engine/Plugins/Media/NDIMedia/Source/NDIMediaEditor/Private/NDIMediaSourceActions.cpp
@@ -0,0 +1,29 @@
+// Copyright Epic Games, Inc. All Rights Reserved.
+
+#include "NDIMediaSourceActions.h"
+#include "NDIMediaSource.h"
+
+#define LOCTEXT_NAMESPACE "NDIMediaSourceActions"
+
+bool FNDIMediaSourceActions::CanFilter()
+{
+	return true;
+}
+
+FText FNDIMediaSourceActions::GetName() const
+{
+	return LOCTEXT("NDIMediaSource_Name", "NDI Media Source");
+}
+
+UClass* FNDIMediaSourceActions::GetSupportedClass() const
+{
+	return UNDIMediaSource::StaticClass();
+}
+
+FColor FNDIMediaSourceActions::GetTypeColor() const
+{
+	// From NDI Brand Guidelines: #6257FF
+	return FColor(98, 87, 255);
+}
+
+#undef LOCTEXT_NAMESPACE
diff --git a/Engine/Plugins/Media/NDIMedia/Source/NDIMediaEditor/Private/NDIMediaSourceActions.h b/Engine/Plugins/Media/NDIMedia/Source/NDIMediaEditor/Private/NDIMediaSourceActions.h
new file mode 100644
--- /dev/null
+++ b/Engine/Plugins/Media/NDIMedia/Source/NDIMediaEditor/Private/NDIMediaSourceActions.h
@@ -0,0 +1,19 @@
+// Copyright Epic Games, Inc. All Rights Reserved.
+
+#pragma once
+
+#include "AssetTools/MediaSourceActions.h"
+
+/**
+ * Implements an action for UNDIMediaSource assets. 
+ */
+class FNDIMediaSourceActions : public FMediaSourceActions
+{
+public:
+	//~ Begin FAssetTypeActions_Base
+	virtual bool CanFilter() override;
+	virtual FText GetName() const override;
+	virtual UClass* GetSupportedClass() const override;
+	virtual FColor GetTypeColor() const override;
+	//~ End FAssetTypeActions_Base
+};
