[NDIMedia] Fix CIS issue 885213

Missing MediaPlayerEditor in plugin dependency list.

diff --git a/Engine/Plugins/Media/NDIMedia/NDIMedia.uplugin b/Engine/Plugins/Media/NDIMedia/NDIMedia.uplugin
--- a/Engine/Plugins/Media/NDIMedia/NDIMedia.uplugin
+++ b/Engine/Plugins/Media/NDIMedia/NDIMedia.uplugin
@@ -39,6 +39,10 @@
 		{
 			"Name": "MediaIOFramework",
 			"Enabled": true
+		},
+		{
+			"Name": "MediaPlayerEditor",
+			"Enabled": true
 		}
 	]
 }
