[NDIMedia] Automatically reopening the NDI stream when media source settings are modified in the media source editor.

diff --git a/Engine/Plugins/Media/NDIMedia/Source/NDIMedia/Private/NDIMediaSource.cpp b/Engine/Plugins/Media/NDIMedia/Source/NDIMedia/Private/NDIMediaSource.cpp
--- a/Engine/Plugins/Media/NDIMedia/Source/NDIMedia/Private/NDIMediaSource.cpp
+++ b/Engine/Plugins/Media/NDIMedia/Source/NDIMedia/Private/NDIMediaSource.cpp
@@ -8,6 +8,10 @@
 #include "NDIMediaModule.h"
 #include "Player/NDIMediaSourceOptions.h"
 
+#if WITH_EDITOR
+UNDIMediaSource::FOnOptionChanged UNDIMediaSource::OnOptionChanged;
+#endif
+
 UNDIMediaSource::UNDIMediaSource()
 	: UCaptureCardMediaSource()
 	, bCaptureAncillary(false)
@@ -207,6 +211,14 @@ void UNDIMediaSource::PostLoad()
 	AssignDefaultConfiguration();
 }
 
+#if WITH_EDITOR
+void UNDIMediaSource::PostEditChangeProperty(FPropertyChangedEvent& InPropertyChangedEvent)
+{
+	OnOptionChanged.Broadcast(this, InPropertyChangedEvent);
+	Super::PostEditChangeProperty(InPropertyChangedEvent);
+}
+#endif
+
 void UNDIMediaSource::AssignDefaultConfiguration()
 {
 	if (!MediaConfiguration.IsValid())
diff --git a/Engine/Plugins/Media/NDIMedia/Source/NDIMedia/Private/Player/NDIMediaPlayer.cpp b/Engine/Plugins/Media/NDIMedia/Source/NDIMedia/Private/Player/NDIMediaPlayer.cpp
--- a/Engine/Plugins/Media/NDIMedia/Source/NDIMedia/Private/Player/NDIMediaPlayer.cpp
+++ b/Engine/Plugins/Media/NDIMedia/Source/NDIMedia/Private/Player/NDIMediaPlayer.cpp
@@ -11,6 +11,7 @@
 #include "MediaIOCoreTextureSampleBase.h"
 #include "NDIMediaLog.h"
 #include "NDIMediaModule.h"
+#include "NDIMediaSource.h"
 #include "NDIMediaSourceOptions.h"
 #include "NDIMediaTextureSample.h"
 #include "NDIMediaTextureSampleConverter.h"
@@ -62,6 +63,26 @@ FGuid FNDIMediaPlayer::GetPlayerPluginGUID() const
 	return FNDIMediaModule::PlayerPluginGUID;
 }
 
+#if WITH_EDITOR	
+void FNDIMediaPlayer::OnOptionsChanged(UObject* InOptions, FPropertyChangedEvent& InPropertyChanged)
+{
+	if (OptionsObject == InOptions)
+	{
+		if (UMediaSource* InMediaSource = Cast<UMediaSource>(InOptions))
+		{
+			// todo: some options can possibly be modified without needing a complete reset.
+			// For now, handle any options changed by restarting the player.
+			{
+				TGuardValue ReopenGuard(bIsReopening, true);
+				Close();
+				Open(InMediaSource->GetUrl(), InMediaSource);
+			}
+		}
+	}
+}
+#endif
+
+
 bool FNDIMediaPlayer::Open(const FString& InUrl, const IMediaOptions* InOptions)
 {
 	if (!Super::Open(InUrl, InOptions))
@@ -69,6 +90,15 @@ bool FNDIMediaPlayer::Open(const FString& InUrl, const IMediaOptions* InOptions)
 		return false;
 	}
 
+#if WITH_EDITOR
+	if (!bIsReopening)
+	{
+		OptionsObject = InOptions->ToUObject();
+		UNDIMediaSource::OnOptionChanged.RemoveAll(this);
+		UNDIMediaSource::OnOptionChanged.AddSP(this, &FNDIMediaPlayer::OnOptionsChanged);
+	}
+#endif
+	
 	MaxNumVideoFrameBuffer = InOptions->GetMediaOption(UE::NDIMediaSourceOptions::MaxVideoFrameBuffer, (int64)8);
 	MaxNumAudioFrameBuffer = InOptions->GetMediaOption(UE::NDIMediaSourceOptions::MaxAudioFrameBuffer, (int64)8);
 	MaxNumMetadataFrameBuffer = InOptions->GetMediaOption(UE::NDIMediaSourceOptions::MaxAncillaryFrameBuffer, (int64)8);
@@ -169,6 +199,14 @@ void FNDIMediaPlayer::Close()
 	AudioSamplePool->Reset();
 	MetadataSamplePool->Reset();
 
+#if WITH_EDITOR
+	if (!bIsReopening)
+	{
+		OptionsObject = nullptr;
+		UNDIMediaSource::OnOptionChanged.RemoveAll(this);
+	}
+#endif
+
 	Super::Close();
 }
 
diff --git a/Engine/Plugins/Media/NDIMedia/Source/NDIMedia/Private/Player/NDIMediaPlayer.h b/Engine/Plugins/Media/NDIMedia/Source/NDIMedia/Private/Player/NDIMediaPlayer.h
--- a/Engine/Plugins/Media/NDIMedia/Source/NDIMedia/Private/Player/NDIMediaPlayer.h
+++ b/Engine/Plugins/Media/NDIMedia/Source/NDIMedia/Private/Player/NDIMediaPlayer.h
@@ -4,6 +4,7 @@
 
 #include "MediaIOCorePlayerBase.h"
 #include "NDIMediaAPI.h"
+#include "UObject/ObjectKey.h"
 
 class FNDIMediaAudioSamplePool;
 class FNDIMediaBinarySamplePool;
@@ -45,6 +46,10 @@ public:
 	//~ End ITimedDataInput
 
 private:
+#if WITH_EDITOR
+	void OnOptionsChanged(UObject* InOptions, FPropertyChangedEvent& InPropertyChanged);
+#endif
+
 	/** Callback handler for video frame reception. */
 	void HandleVideoFrameReceived(FNDIStreamReceiver* InReceiver, const NDIlib_video_frame_v2_t& InVideoFrame, const FTimespan& InTime);
 
@@ -54,6 +59,12 @@ private:
 	void VerifyFrameDropCount();
 
 private:
+	/** Keep track of the options (media source) object this player was opened with to filter global events. */
+	FObjectKey OptionsObject;
+
+	/** Indicate if currently in a reopening sequence. */
+	bool bIsReopening = false;
+	
 	/** Maximum Audio Sample Pool size */
 	int32 MaxNumAudioFrameBuffer = 0;
 	
diff --git a/Engine/Plugins/Media/NDIMedia/Source/NDIMedia/Public/NDIMediaSource.h b/Engine/Plugins/Media/NDIMedia/Source/NDIMedia/Public/NDIMediaSource.h
--- a/Engine/Plugins/Media/NDIMedia/Source/NDIMedia/Public/NDIMediaSource.h
+++ b/Engine/Plugins/Media/NDIMedia/Source/NDIMedia/Public/NDIMediaSource.h
@@ -73,6 +73,16 @@ public:
 	UPROPERTY(BlueprintReadOnly, EditAnywhere, Category="Debug", meta=(DisplayName="Burn Frame Timecode"))
 	bool bEncodeTimecodeInTexel;
 
+#if WITH_EDITOR
+	/**
+	 * Delegate called when an NDI Media Source is modified.
+	 * This can be used to propagate the change to active players.
+	 * @remark Called from the main thread.
+	 */
+	DECLARE_MULTICAST_DELEGATE_TwoParams(FOnOptionChanged, UObject*, FPropertyChangedEvent&);
+	static FOnOptionChanged OnOptionChanged;
+#endif
+	
 	//~ Begin IMediaOptions
 	NDIMEDIA_API virtual bool GetMediaOption(const FName& Key, bool DefaultValue) const override;
 	NDIMEDIA_API virtual int64 GetMediaOption(const FName& Key, int64 DefaultValue) const override;
@@ -88,6 +98,9 @@ public:
 	//~ Begin UObject
 	NDIMEDIA_API virtual void PostLoad() override;
 	NDIMEDIA_API virtual bool SupportsFormatAutoDetection() const override { return true; }
+#if WITH_EDITOR
+	NDIMEDIA_API virtual void PostEditChangeProperty(FPropertyChangedEvent& InPropertyChangedEvent) override;
+#endif
 	//~ End UObject
 
 private:
