[NDIMedia] Fix NDI audio input

After investigating the frame synchronization api, it seems to not be working as expected. The returned audio samples from framesync_capture_audio are all silence (0). Testing the NDI 6 SDK's NDIlib_Recv_FrameSync example also produces the same result.

The non-sync api: recv_capture_v2 seems to work fine however, so we'll be going with that until further notice.

diff --git a/Engine/Plugins/Media/NDIMedia/Source/NDIMedia/Private/Player/NDIStreamReceiver.cpp b/Engine/Plugins/Media/NDIMedia/Source/NDIMedia/Private/Player/NDIStreamReceiver.cpp
--- a/Engine/Plugins/Media/NDIMedia/Source/NDIMedia/Private/Player/NDIStreamReceiver.cpp
+++ b/Engine/Plugins/Media/NDIMedia/Source/NDIMedia/Private/Player/NDIStreamReceiver.cpp
@@ -12,6 +12,14 @@
 #include "Editor.h"
 #endif
 
+// Reference: https://docs.ndi.video/all/developing-with-ndi/advanced-sdk/ndi-sdk-review/video-formats/frame-synchronization
+// This is only exposed as a CVar for now because it is not working correctly.
+static TAutoConsoleVariable<bool> CVarNDIUseFrameSync(
+	TEXT("NDIMediaReceiver.UseFrameSync"),
+	false, // Disabled by default until issues are fixed.
+	TEXT("Use the ndi frame synchronization api to capture video and audio. Note: need to restart the streams for this cvar to take effect."),
+	ECVF_Default);
+
 namespace UE::NDIStreamReceiver::Private
 {
 	NDIlib_recv_bandwidth_e GetNDIBandwidth(ENDIReceiverBandwidth InBandwidth)
@@ -47,9 +55,9 @@ namespace UE::NDIStreamReceiver::Private
 
 FNDIStreamReceiver::FNDIStreamReceiver(const TSharedPtr<FNDIMediaRuntimeLibrary>& InNdiLib)
 	: Resolution(0, 0)
+	, bUseFrameSync(CVarNDIUseFrameSync.GetValueOnAnyThread())
 	, NdiLib(InNdiLib)
 {
-	
 }
 
 FNDIStreamReceiver::~FNDIStreamReceiver()
@@ -152,7 +160,10 @@ void FNDIStreamReceiver::StartConnection()
 		ReceiveInstance = NewReceiveInstance;
 
 		// create a new frame sync instance
-		FrameSyncInstance = NdiLib->Lib->framesync_create(ReceiveInstance);
+		if (bUseFrameSync)
+		{
+			FrameSyncInstance = NdiLib->Lib->framesync_create(ReceiveInstance);
+		}
 	}
 }
 
@@ -201,7 +212,7 @@ void FNDIStreamReceiver::ChangeConnection(const FNDISourceSettings& InSourceSett
 
 			if (SourceSettings.IsValid())
 			{
-				if (bSettingsChanged || !ReceiveInstance || !FrameSyncInstance)
+				if (bSettingsChanged || !ReceiveInstance || (bUseFrameSync && !FrameSyncInstance))
 				{
 					// Connection information is valid, and something has changed that requires the connection to be remade
 					StartConnection();
@@ -222,15 +233,22 @@ int32 FNDIStreamReceiver::GetAudioChannels()
 
 	int32 NumChannels = 0;
 
-	if (FrameSyncInstance && SourceSettings.bCaptureAudio)
+	if (SourceSettings.bCaptureAudio)
 	{
-		int AvailableNumFrames = NdiLib->Lib->framesync_audio_queue_depth(FrameSyncInstance);	// Samples per channel
+		if (bUseFrameSync && FrameSyncInstance)
+		{
+			int AvailableNumFrames = NdiLib->Lib->framesync_audio_queue_depth(FrameSyncInstance);	// Samples per channel
 
-		if (AvailableNumFrames > 0)
+			if (AvailableNumFrames > 0)
+			{
+				NDIlib_audio_frame_v2_t AudioFrame;
+				NdiLib->Lib->framesync_capture_audio(FrameSyncInstance, &AudioFrame, 0, 0, 0);
+				NumChannels = AudioFrame.no_channels;
+			}
+		}
+		else
 		{
-			NDIlib_audio_frame_v2_t AudioFrame;
-			NdiLib->Lib->framesync_capture_audio(FrameSyncInstance, &AudioFrame, 0, 0, 0);
-			NumChannels = AudioFrame.no_channels;
+			NumChannels = LastNumAudioChannels;
 		}
 	}
 
@@ -315,80 +333,99 @@ void FNDIStreamReceiver::Shutdown()
 
 bool FNDIStreamReceiver::FetchVideo(const FTimespan& InTime)
 {
-	// This function is called on the Engine's Main Rendering Thread. Be very careful when doing stuff here.
-	// Make sure things are done quick and efficient.
+	// This function can either be called from the game thread or rendering thread.
+	// todo: dedicated thread for the frame copy.
 
 	// Ensure thread safety
 	FScopeLock Lock(&RenderSyncContext);
 
-	bool bFrameReceived = false;
-
 	// check for our frame sync object and that we are actually connected to the end point
-	if (FrameSyncInstance && SourceSettings.bCaptureVideo)
+	if (!SourceSettings.bCaptureVideo || (bUseFrameSync && !FrameSyncInstance) || !ReceiveInstance)
 	{
-		// Using a frame-sync we can always get data which is the magic and it will adapt
-		// to the frame-rate that it is being called with.
-		NDIlib_video_frame_v2_t VideoFrame;
-		NdiLib->Lib->framesync_capture_video(FrameSyncInstance, &VideoFrame, NDIlib_frame_format_type_progressive);
+		return false;
+	}
 
-		// Update our Performance Metrics
-		GatherPerformanceMetrics();
+	NDIlib_video_frame_v2_t VideoFrame;
 
-		if (VideoFrame.p_data)
+	if (bUseFrameSync && FrameSyncInstance)
+	{
+		NdiLib->Lib->framesync_capture_video(FrameSyncInstance, &VideoFrame, NDIlib_frame_format_type_progressive);
+	}
+	else
+	{
+		const NDIlib_frame_type_e FrameType = NdiLib->Lib->recv_capture_v2(ReceiveInstance, &VideoFrame, nullptr, nullptr, 0);
+		if (FrameType != NDIlib_frame_type_video)
 		{
-			// Ensure that we inform all those interested when the stream starts up
-			SetIsCurrentlyConnected(/*bInConnected*/ true);
+			return false;
+		}
+	}
 
-			// Update the Framerate, if it has changed
-			FrameRate.Numerator = VideoFrame.frame_rate_N;
-			FrameRate.Denominator = VideoFrame.frame_rate_D;
+	// Update our Performance Metrics
+	GatherPerformanceMetrics();
 
-			// Update the Resolution
-			Resolution.X = VideoFrame.xres;
-			Resolution.Y = VideoFrame.yres;
+	bool bFrameReceived = false;
 
-			if (bSyncTimecodeToSource)
-			{
-				// Update the timecode from the current 'SourceTime' value
-				int64_t SourceTime = VideoFrame.timecode % 864000000000; // Modulo the number of 100ns intervals in 24 hours
-				Timecode = UE::NDIStreamReceiver::Private::ToTimecode(SourceTime, FrameRate);
-			}
-			else
-			{
-				// Update the timecode from the current 'SystemTime' value
-				int64_t SystemTime = FDateTime::Now().GetTimeOfDay().GetTicks();
-				Timecode = UE::NDIStreamReceiver::Private::ToTimecode(SystemTime, FrameRate);
-				// todo: (validate) use provided engine time.
-				//Timecode = UE::NDIStreamReceiver::Private::ToTimecode(InTime, FrameRate);
-			}
+	if (VideoFrame.p_data)
+	{
+		// Ensure that we inform all those interested when the stream starts up
+		SetIsCurrentlyConnected(/*bInConnected*/ true);
 
-			// Redraw if:
-			// - timestamp is undefined, or
-			// - timestamp has changed, or
-			// - frame format type has changed (e.g. different field)
-			if ((VideoFrame.timestamp == NDIlib_recv_timestamp_undefined) ||
-				(VideoFrame.timestamp != LastFrameTimestamp) ||
-				(VideoFrame.frame_format_type != LastFrameFormatType))
-			{
-				bFrameReceived = true;
+		// Update the Framerate, if it has changed
+		FrameRate.Numerator = VideoFrame.frame_rate_N;
+		FrameRate.Denominator = VideoFrame.frame_rate_D;
 
-				LastFrameTimestamp = VideoFrame.timestamp;
-				LastFrameFormatType = VideoFrame.frame_format_type;
+		// Update the Resolution
+		Resolution.X = VideoFrame.xres;
+		Resolution.Y = VideoFrame.yres;
 
-				OnVideoFrameReceived.Broadcast(this, VideoFrame, InTime);
+		if (bSyncTimecodeToSource)
+		{
+			// Update the timecode from the current 'SourceTime' value
+			int64_t SourceTime = VideoFrame.timecode % 864000000000; // Modulo the number of 100ns intervals in 24 hours
+			Timecode = UE::NDIStreamReceiver::Private::ToTimecode(SourceTime, FrameRate);
+		}
+		else
+		{
+			// Update the timecode from the current 'SystemTime' value
+			int64_t SystemTime = FDateTime::Now().GetTimeOfDay().GetTicks();
+			Timecode = UE::NDIStreamReceiver::Private::ToTimecode(SystemTime, FrameRate);
+			// todo: (validate) use provided engine time.
+			//Timecode = UE::NDIStreamReceiver::Private::ToTimecode(InTime, FrameRate);
+		}
 
-				if (VideoFrame.p_metadata)
-				{
-					FString Data(UTF8_TO_TCHAR(VideoFrame.p_metadata));
-					OnMetaDataReceived.Broadcast(this, Data, true);
-				}
+		// Redraw if:
+		// - timestamp is undefined, or
+		// - timestamp has changed, or
+		// - frame format type has changed (e.g. different field)
+		if ((VideoFrame.timestamp == NDIlib_recv_timestamp_undefined) ||
+			(VideoFrame.timestamp != LastFrameTimestamp) ||
+			(VideoFrame.frame_format_type != LastFrameFormatType))
+		{
+			bFrameReceived = true;
+
+			LastFrameTimestamp = VideoFrame.timestamp;
+			LastFrameFormatType = VideoFrame.frame_format_type;
+
+			OnVideoFrameReceived.Broadcast(this, VideoFrame, InTime);
+
+			if (VideoFrame.p_metadata)
+			{
+				FString Data(UTF8_TO_TCHAR(VideoFrame.p_metadata));
+				OnMetaDataReceived.Broadcast(this, Data, true);
 			}
 		}
+	}
 
-		// Release the video. You could keep the frame if you want and release it later.
+	// Release the video. You could keep the frame if you want and release it later.
+	if (bUseFrameSync && FrameSyncInstance)
+	{
 		NdiLib->Lib->framesync_free_video(FrameSyncInstance, &VideoFrame);
 	}
-
+	else
+	{
+		NdiLib->Lib->recv_free_video_v2(ReceiveInstance, &VideoFrame);
+	}
+	
 	return bFrameReceived;
 }
 
@@ -396,34 +433,54 @@ bool FNDIStreamReceiver::FetchAudio(const FTimespan& InTime)
 {
 	FScopeLock Lock(&AudioSyncContext);
 
-	bool bFrameReceived = false;
+	if (!SourceSettings.bCaptureAudio || (bUseFrameSync && !FrameSyncInstance) || !ReceiveInstance)
+	{
+		return false;
+	}
+
+	NDIlib_audio_frame_v2_t AudioFrame;
 
-	if (FrameSyncInstance && SourceSettings.bCaptureAudio)
+	if (bUseFrameSync && FrameSyncInstance)
 	{
+		// fix me: this function currently always return silence. It is also the case in ndi sdk examples. Disabled for now.
 		int NumSamples = NdiLib->Lib->framesync_audio_queue_depth(FrameSyncInstance);
-
-		// Using a frame-sync we can always get data which is the magic and it will adapt
-		// to the frame-rate that it is being called with.
-		NDIlib_audio_frame_v2_t AudioFrame;
 		NdiLib->Lib->framesync_capture_audio(FrameSyncInstance, &AudioFrame, 0, 0, NumSamples);
-		
-		if (AudioFrame.p_data)
+	}
+	else
+	{
+		const NDIlib_frame_type_e FrameType = NdiLib->Lib->recv_capture_v2(ReceiveInstance, nullptr, &AudioFrame, nullptr, 0);
+		if (FrameType != NDIlib_frame_type_audio)
 		{
-			// Ensure that we inform all those interested when the stream starts up
-			SetIsCurrentlyConnected(/*bInConnected*/ true);
+			return false;
+		}
+	}
 
-			const int32 NumAvailableSamples = AudioFrame.no_samples * AudioFrame.no_channels;
+	bool bFrameReceived = false;
 
-			if (NumAvailableSamples > 0)
-			{
-				bFrameReceived = true;
-				OnAudioFrameReceived.Broadcast(this, AudioFrame, InTime);
-			}
+	if (AudioFrame.p_data)
+	{
+		// Ensure that we inform all those interested when the stream starts up
+		SetIsCurrentlyConnected(/*bInConnected*/ true);
+
+		const int32 NumAvailableSamples = AudioFrame.no_samples * AudioFrame.no_channels;
+
+		if (NumAvailableSamples > 0)
+		{
+			LastNumAudioChannels = AudioFrame.no_channels;
+			bFrameReceived = true;
+			OnAudioFrameReceived.Broadcast(this, AudioFrame, InTime);
 		}
+	}
 
-		// Release the audio frame
+	// Release the audio frame
+	if (bUseFrameSync && FrameSyncInstance)
+	{
 		NdiLib->Lib->framesync_free_audio(FrameSyncInstance, &AudioFrame);
 	}
+	else
+	{
+		NdiLib->Lib->recv_free_audio_v2(ReceiveInstance, &AudioFrame);
+	}
 
 	return bFrameReceived;
 }
diff --git a/Engine/Plugins/Media/NDIMedia/Source/NDIMedia/Private/Player/NDIStreamReceiver.h b/Engine/Plugins/Media/NDIMedia/Source/NDIMedia/Private/Player/NDIStreamReceiver.h
--- a/Engine/Plugins/Media/NDIMedia/Source/NDIMedia/Private/Player/NDIStreamReceiver.h
+++ b/Engine/Plugins/Media/NDIMedia/Source/NDIMedia/Private/Player/NDIStreamReceiver.h
@@ -167,6 +167,9 @@ private:
 	/** Indicates whether the timecode should be synced to the Source Timecode value */
 	bool bSyncTimecodeToSource = true;
 
+	/** Enables the use of frame sync.  */
+	bool bUseFrameSync = false;
+
 	/** Information describing detailed information about the sender this receiver is currently connected to */
 	FNDISourceSettings SourceSettings;
 
@@ -179,6 +182,9 @@ private:
 	/** Keep track of the last VideoFrame format type received. Used to detect new frame. */
 	NDIlib_frame_format_type_e LastFrameFormatType = NDIlib_frame_format_type_max;
 
+	/** Keep track of the last AudioFrame's num audio channel received. */
+	int32 LastNumAudioChannels = 0;
+
 	/** Keep track of the current connection status with the source. */
 	bool bIsCurrentlyConnected = false;
 
