[NDIMedia] Fix missing CaptureRHIResource api for nDisplay support

diff --git a/Engine/Plugins/Media/NDIMedia/Source/NDIMedia/Private/NDIMediaCapture.cpp b/Engine/Plugins/Media/NDIMedia/Source/NDIMedia/Private/NDIMediaCapture.cpp
--- a/Engine/Plugins/Media/NDIMedia/Source/NDIMedia/Private/NDIMediaCapture.cpp
+++ b/Engine/Plugins/Media/NDIMedia/Source/NDIMedia/Private/NDIMediaCapture.cpp
@@ -403,6 +403,17 @@ bool UNDIMediaCapture::PostInitializeCaptureRenderTarget(UTextureRenderTarget2D*
 	return bSuccess;
 }
 
+bool UNDIMediaCapture::PostInitializeCaptureRHIResource(const FRHICaptureResourceDescription& InResourceDescription)
+{
+	const bool bSuccess = StartNewCapture();
+	if (bSuccess)
+	{
+		UE_LOG(LogNDIMedia, Log, TEXT("Media Capture Started: Render Target (%d x %d)."), 
+			InResourceDescription.ResourceSize.X, InResourceDescription.ResourceSize.Y);
+	}
+	return bSuccess;
+}
+
 bool UNDIMediaCapture::UpdateAudioDeviceImpl(const FAudioDeviceHandle& InAudioDeviceHandle)
 {
 	FScopeLock ScopeLock(&CaptureInstanceCriticalSection);
diff --git a/Engine/Plugins/Media/NDIMedia/Source/NDIMedia/Public/NDIMediaCapture.h b/Engine/Plugins/Media/NDIMedia/Source/NDIMedia/Public/NDIMediaCapture.h
--- a/Engine/Plugins/Media/NDIMedia/Source/NDIMedia/Public/NDIMediaCapture.h
+++ b/Engine/Plugins/Media/NDIMedia/Source/NDIMedia/Public/NDIMediaCapture.h
@@ -23,6 +23,7 @@ protected:
 	virtual bool InitializeCapture() override;
 	virtual bool PostInitializeCaptureViewport(TSharedPtr<FSceneViewport>& InSceneViewport) override;
 	virtual bool PostInitializeCaptureRenderTarget(UTextureRenderTarget2D* InRenderTarget) override;
+	virtual bool PostInitializeCaptureRHIResource(const FRHICaptureResourceDescription& InResourceDescription) override;
 	virtual bool UpdateAudioDeviceImpl(const FAudioDeviceHandle& InAudioDeviceHandle) override;
 	virtual void StopCaptureImpl(bool bAllowPendingFrameToBeProcess) override;
 
