[NDIMedia] Renaming headers and classes to avoid conflict with published Newtek plugin

diff --git a/Engine/Plugins/Media/NDIMedia/Source/NDIMedia/Private/NDIMediaModule.cpp b/Engine/Plugins/Media/NDIMedia/Source/NDIMedia/Private/NDIMediaModule.cpp
--- a/Engine/Plugins/Media/NDIMedia/Source/NDIMedia/Private/NDIMediaModule.cpp
+++ b/Engine/Plugins/Media/NDIMedia/Source/NDIMedia/Private/NDIMediaModule.cpp
@@ -15,7 +15,7 @@
 #include "NDIMediaLog.h"
 #include "NDIMediaSettings.h"
 #include "NDISourceFinder.h"
-#include "Player/NDIMediaPlayer.h"
+#include "Player/NDIMediaStreamPlayer.h"
 #include "Player/NDIStreamReceiverManager.h"
 
 DEFINE_LOG_CATEGORY(LogNDIMedia);
@@ -202,7 +202,7 @@ int32 FNDIMediaModule::GetPlayabilityConfidenceScore(const FString& Url, const I
 
 TSharedPtr<IMediaPlayer> FNDIMediaModule::CreatePlayer(IMediaEventSink& InEventSink)
 {
-	return MakeShared<FNDIMediaPlayer>(InEventSink);
+	return MakeShared<FNDIMediaStreamPlayer>(InEventSink);
 }
 
 FText FNDIMediaModule::GetDisplayName() const
diff --git a/Engine/Plugins/Media/NDIMedia/Source/NDIMedia/Private/NDIMediaTimecodeProvider.cpp b/Engine/Plugins/Media/NDIMedia/Source/NDIMedia/Private/NDIMediaTimecodeProvider.cpp
new file mode 100644
--- /dev/null
+++ b/Engine/Plugins/Media/NDIMedia/Source/NDIMedia/Private/NDIMediaTimecodeProvider.cpp
@@ -0,0 +1,110 @@
+// Copyright Epic Games, Inc. All Rights Reserved.
+
+#include "NDIMediaTimecodeProvider.h"
+
+#include "NDIMediaModule.h"
+#include "Player/NDIStreamReceiver.h"
+#include "Player/NDIStreamReceiverManager.h"
+
+bool UNDIMediaTimecodeProvider::FetchTimecode(FQualifiedFrameTime& OutFrameTime)
+{
+	FScopeLock Lock(&StateSyncContext);
+
+	if (!Receiver ||
+	    (GetSynchronizationState() != ETimecodeProviderSynchronizationState::Synchronized))
+	{
+		return false;
+	}
+
+	OutFrameTime = MostRecentFrameTime;
+
+	return true;
+}
+
+ETimecodeProviderSynchronizationState UNDIMediaTimecodeProvider::GetSynchronizationState() const
+{
+	FScopeLock Lock(&StateSyncContext);
+
+	if (!Receiver)
+	{
+		return ETimecodeProviderSynchronizationState::Closed;
+	}
+
+	return State;
+}
+
+bool UNDIMediaTimecodeProvider::Initialize(UEngine* InEngine)
+{
+	State = ETimecodeProviderSynchronizationState::Closed;
+
+	FNDISourceSettings SourceSettings;
+	SourceSettings.Bandwidth = Bandwidth;
+	SourceSettings.bCaptureAudio = false;
+	SourceSettings.bCaptureVideo = true;
+	SourceSettings.SourceName = TimecodeConfiguration.MediaConfiguration.MediaConnection.Device.DeviceName.ToString();
+
+	// Check if the receiver is already created by another object.
+	if (FNDIMediaModule* Module = FNDIMediaModule::Get())
+	{
+		FNDIStreamReceiverManager& StreamReceiverManager = Module->GetStreamReceiverManager();
+		Receiver = StreamReceiverManager.FindReceiver(SourceSettings.SourceName);
+		if (Receiver)
+		{
+			SourceSettings.bCaptureAudio = Receiver->GetCurrentSourceSettings().bCaptureAudio;
+		}
+	}
+
+	if (!Receiver)
+	{
+		Receiver = MakeShared<FNDIStreamReceiver>(FNDIMediaModule::GetNDIRuntimeLibrary());	
+	}
+
+	Receiver->SetSyncTimecodeToSource(/* bInSyncTimecodeToSource*/true);
+	Receiver->Initialize(SourceSettings, FNDIStreamReceiver::ECaptureMode::OnEndFrameRT);
+
+	VideoFrameReceivedHandle = Receiver->OnVideoFrameReceived.AddLambda([this](FNDIStreamReceiver* InReceiver, const NDIlib_video_frame_v2_t& InVideoFrame, const FTimespan& InTimecode)
+	{
+		FScopeLock Lock(&StateSyncContext);
+		State = ETimecodeProviderSynchronizationState::Synchronized;
+		MostRecentFrameTime = FQualifiedFrameTime(Receiver->GetCurrentTimecode(), Receiver->GetCurrentFrameRate());
+	});
+	ConnectedHandle = Receiver->OnConnected.AddLambda([this](FNDIStreamReceiver* InReceiver)
+	{
+		FScopeLock Lock(&StateSyncContext);
+		State = ETimecodeProviderSynchronizationState::Synchronizing;
+	});
+	DisconnectedHandle = Receiver->OnDisconnected.AddLambda([this](FNDIStreamReceiver* InReceiver)
+	{
+		FScopeLock Lock(&StateSyncContext);
+		State = ETimecodeProviderSynchronizationState::Closed;
+	});
+
+	return true;
+}
+
+void UNDIMediaTimecodeProvider::Shutdown(UEngine* InEngine)
+{
+	ReleaseResources();
+}
+
+void UNDIMediaTimecodeProvider::BeginDestroy()
+{
+	ReleaseResources();
+
+	Super::BeginDestroy();
+}
+
+void UNDIMediaTimecodeProvider::ReleaseResources()
+{
+	if (Receiver)
+	{
+		Receiver->OnVideoFrameReceived.Remove(VideoFrameReceivedHandle);
+		Receiver->OnConnected.Remove(ConnectedHandle);
+		Receiver->OnDisconnected.Remove(DisconnectedHandle);
+	}
+	VideoFrameReceivedHandle.Reset();
+	ConnectedHandle.Reset();
+	DisconnectedHandle.Reset();
+
+	State = ETimecodeProviderSynchronizationState::Closed;
+}
diff --git a/Engine/Plugins/Media/NDIMedia/Source/NDIMedia/Private/NDITimecodeProvider.cpp b/Engine/Plugins/Media/NDIMedia/Source/NDIMedia/Private/NDITimecodeProvider.cpp
deleted file mode 100644
diff --git a/Engine/Plugins/Media/NDIMedia/Source/NDIMedia/Private/Player/NDIMediaPlayer.cpp b/Engine/Plugins/Media/NDIMedia/Source/NDIMedia/Private/Player/NDIMediaPlayer.cpp
deleted file mode 100644
diff --git a/Engine/Plugins/Media/NDIMedia/Source/NDIMedia/Private/Player/NDIMediaPlayer.h b/Engine/Plugins/Media/NDIMedia/Source/NDIMedia/Private/Player/NDIMediaPlayer.h
deleted file mode 100644
diff --git a/Engine/Plugins/Media/NDIMedia/Source/NDIMedia/Private/Player/NDIMediaReceiverPerformanceData.h b/Engine/Plugins/Media/NDIMedia/Source/NDIMedia/Private/Player/NDIMediaReceiverPerformanceData.h
new file mode 100644
--- /dev/null
+++ b/Engine/Plugins/Media/NDIMedia/Source/NDIMedia/Private/Player/NDIMediaReceiverPerformanceData.h
@@ -0,0 +1,48 @@
+// Copyright Epic Games, Inc. All Rights Reserved.
+
+#pragma once
+
+#include "NDIMediaReceiverPerformanceData.generated.h"
+
+/**
+ *  NDI receiver's performance statistics on number of frames received and dropped.
+ */
+USTRUCT(BlueprintType)
+struct FNDIMediaReceiverPerformanceData
+{
+	GENERATED_BODY()
+
+public:
+	/** The number of video frames received */
+	UPROPERTY(BlueprintReadOnly, VisibleAnywhere, Category = "Information")
+	int64 VideoFrames = 0;
+
+	/** The number of video frames dropped */
+	UPROPERTY(BlueprintReadonly, VisibleAnywhere, Category = "Information")
+	int64 DroppedVideoFrames = 0;
+
+	/** The number of audio frames received */
+	UPROPERTY(BlueprintReadOnly, VisibleAnywhere, Category = "Information")
+	int64 AudioFrames = 0;
+
+	/** The number of audio frames dropped */
+	UPROPERTY(BlueprintReadOnly, VisibleAnywhere, Category = "Information")
+	int64 DroppedAudioFrames = 0;
+
+	/** The number of metadata frames received */
+	UPROPERTY(BlueprintReadOnly, VisibleAnywhere, Category = "Information")
+	int64 MetadataFrames = 0;
+
+	/** The number of metadata frames dropped */
+	UPROPERTY(BlueprintReadOnly, VisibleAnywhere, Category = "Information")
+	int64 DroppedMetadataFrames = 0;
+
+	/** Compares this object to 'other' and returns a determination of whether they are equal */
+	bool operator==(const FNDIMediaReceiverPerformanceData& other) const
+	{
+		return AudioFrames == other.AudioFrames && DroppedAudioFrames == other.DroppedAudioFrames &&
+			   DroppedMetadataFrames == other.DroppedMetadataFrames &&
+			   DroppedVideoFrames == other.DroppedVideoFrames && MetadataFrames == other.MetadataFrames &&
+			   VideoFrames == other.VideoFrames;
+	}
+};
\ No newline at end of file
diff --git a/Engine/Plugins/Media/NDIMedia/Source/NDIMedia/Private/Player/NDIMediaStreamPlayer.cpp b/Engine/Plugins/Media/NDIMedia/Source/NDIMedia/Private/Player/NDIMediaStreamPlayer.cpp
new file mode 100644
--- /dev/null
+++ b/Engine/Plugins/Media/NDIMedia/Source/NDIMedia/Private/Player/NDIMediaStreamPlayer.cpp
@@ -0,0 +1,447 @@
+// Copyright Epic Games, Inc. All Rights Reserved.
+
+#include "NDIMediaStreamPlayer.h"
+
+#include "IMediaEventSink.h"
+#include "IMediaOptions.h"
+#include "MediaIOCoreAudioSampleBase.h"
+#include "MediaIOCoreBinarySampleBase.h"
+#include "MediaIOCoreEncodeTime.h"
+#include "MediaIOCoreSamples.h"
+#include "MediaIOCoreTextureSampleBase.h"
+#include "NDIMediaLog.h"
+#include "NDIMediaModule.h"
+#include "NDIMediaSource.h"
+#include "NDIMediaSourceOptions.h"
+#include "NDIMediaTextureSample.h"
+#include "NDIMediaTextureSampleConverter.h"
+#include "NDIStreamReceiver.h"
+#include "NDIStreamReceiverManager.h"
+
+#define LOCTEXT_NAMESPACE "NDIMediaPlayer"
+
+class FNDIMediaTextureSamplePool : public TMediaObjectPool<FNDIMediaTextureSample>
+{};
+
+/**
+ * Implements a media audio sample for NDIMedia
+ */
+class FNDIMediaAudioSample : public FMediaIOCoreAudioSampleBase
+{
+	using Super = FMediaIOCoreAudioSampleBase;
+
+public:
+};
+
+class FNDIMediaAudioSamplePool : public TMediaObjectPool<FNDIMediaAudioSample>
+{};
+
+/*
+ * Implements a pool for NDI binary sample objects. 
+ */
+class FNDIMediaBinarySamplePool : public TMediaObjectPool<FMediaIOCoreBinarySampleBase> { };
+
+FNDIMediaStreamPlayer::FNDIMediaStreamPlayer(IMediaEventSink& InEventSink)
+	: Super(InEventSink)
+	, NDIPlayerState(EMediaState::Closed)
+	, EventSink(InEventSink)
+	, TextureSamplePool(new FNDIMediaTextureSamplePool)
+	, AudioSamplePool(new FNDIMediaAudioSamplePool)
+	, MetadataSamplePool(new FNDIMediaBinarySamplePool)
+{}
+
+FNDIMediaStreamPlayer::~FNDIMediaStreamPlayer()
+{
+	Close();
+
+	TextureSamplePool.Reset();
+	AudioSamplePool.Reset();
+}
+
+FGuid FNDIMediaStreamPlayer::GetPlayerPluginGUID() const
+{
+	return FNDIMediaModule::PlayerPluginGUID;
+}
+
+#if WITH_EDITOR	
+void FNDIMediaStreamPlayer::OnOptionsChanged(UObject* InOptions, FPropertyChangedEvent& InPropertyChanged)
+{
+	if (OptionsObject == InOptions)
+	{
+		if (UMediaSource* InMediaSource = Cast<UMediaSource>(InOptions))
+		{
+			// todo: some options can possibly be modified without needing a complete reset.
+			// For now, handle any options changed by restarting the player.
+			{
+				TGuardValue ReopenGuard(bIsReopening, true);
+				Close();
+				Open(InMediaSource->GetUrl(), InMediaSource);
+			}
+		}
+	}
+}
+#endif
+
+
+bool FNDIMediaStreamPlayer::Open(const FString& InUrl, const IMediaOptions* InOptions)
+{
+	if (!Super::Open(InUrl, InOptions))
+	{
+		return false;
+	}
+
+#if WITH_EDITOR
+	if (!bIsReopening)
+	{
+		OptionsObject = InOptions->ToUObject();
+		UNDIMediaSource::OnOptionChanged.RemoveAll(this);
+		UNDIMediaSource::OnOptionChanged.AddSP(this, &FNDIMediaStreamPlayer::OnOptionsChanged);
+	}
+#endif
+	
+	MaxNumVideoFrameBuffer = InOptions->GetMediaOption(UE::NDIMediaSourceOptions::MaxVideoFrameBuffer, (int64)8);
+	MaxNumAudioFrameBuffer = InOptions->GetMediaOption(UE::NDIMediaSourceOptions::MaxAudioFrameBuffer, (int64)8);
+	MaxNumMetadataFrameBuffer = InOptions->GetMediaOption(UE::NDIMediaSourceOptions::MaxAncillaryFrameBuffer, (int64)8);
+	bEncodeTimecodeInTexel = InOptions->GetMediaOption(UE::NDIMediaSourceOptions::EncodeTimecodeInTexel, false);
+
+	// Setup our different supported channels based on source settings
+	SetupSampleChannels();
+
+	// configure format information for base class
+	AudioTrackFormat.BitsPerSample = 32;
+	AudioTrackFormat.NumChannels = 0;
+	AudioTrackFormat.SampleRate = 44100;
+	AudioTrackFormat.TypeName = FString(TEXT("PCM"));
+
+	bCaptureVideo = InOptions->GetMediaOption(UE::NDIMediaSourceOptions::CaptureVideo, true);
+	bCaptureAudio = InOptions->GetMediaOption(UE::NDIMediaSourceOptions::CaptureAudio, false);
+	bCaptureAncillary = InOptions->GetMediaOption(UE::NDIMediaSourceOptions::CaptureAncillary, false);
+	SupportedSampleTypes = bCaptureVideo ? EMediaIOSampleType::Video : EMediaIOSampleType::None;
+	SupportedSampleTypes |= bCaptureAudio ? EMediaIOSampleType::Audio : EMediaIOSampleType::None;
+	SupportedSampleTypes |= bCaptureAncillary ? EMediaIOSampleType::Metadata : EMediaIOSampleType::None;
+	Samples->EnableTimedDataChannels(this, SupportedSampleTypes);
+
+	FNDISourceSettings SourceSettings;
+	SourceSettings.Bandwidth = static_cast<ENDIReceiverBandwidth>(InOptions->GetMediaOption(UE::NDIMediaSourceOptions::Bandwidth, static_cast<int64>(SourceSettings.Bandwidth)));
+	SourceSettings.bCaptureAudio = bCaptureAudio;
+	SourceSettings.bCaptureVideo = bCaptureVideo;
+
+	FString Scheme;
+	FString Location;
+	if (InUrl.Split(TEXT("://"), &Scheme, &Location, ESearchCase::CaseSensitive))
+	{
+		SourceSettings.SourceName = Location;
+	}
+
+	// Check if the receiver is already created by another object.
+	if (FNDIMediaModule* Module = FNDIMediaModule::Get())
+	{
+		FNDIStreamReceiverManager& StreamReceiverManager = Module->GetStreamReceiverManager();
+		Receiver = StreamReceiverManager.FindReceiver(SourceSettings.SourceName);
+
+		if (!Receiver)
+		{
+			Receiver = MakeShared<FNDIStreamReceiver>(FNDIMediaModule::GetNDIRuntimeLibrary());
+		}
+	}
+	
+	if (!Receiver)
+	{
+		UE_LOG(LogNDIMedia, Error, TEXT("Failed to acquire NDI receiver."));
+		return false;
+	}
+	
+	// Hook into the video and audio captures
+	VideoReceivedHandle = Receiver->OnVideoFrameReceived.AddRaw(this, &FNDIMediaStreamPlayer::HandleVideoFrameReceived);
+	AudioReceivedHandle = Receiver->OnAudioFrameReceived.AddRaw(this, &FNDIMediaStreamPlayer::HandleAudioFrameReceived);
+
+	// Control the player's state based on the receiver connecting and disconnecting
+	ConnectedHandle = Receiver->OnConnected.AddLambda([this](FNDIStreamReceiver* receiver)
+	{
+		NDIPlayerState = EMediaState::Playing;
+	});
+	DisconnectedHandle = Receiver->OnDisconnected.AddLambda([this](FNDIStreamReceiver* receiver)
+	{
+		NDIPlayerState = EMediaState::Closed;
+	});
+
+	// Get ready to connect
+	CurrentState = EMediaState::Preparing;
+	NDIPlayerState = EMediaState::Preparing;
+	EventSink.ReceiveMediaEvent(EMediaEvent::MediaConnecting);
+	
+	Receiver->SetSyncTimecodeToSource(InOptions->GetMediaOption(UE::NDIMediaSourceOptions::SyncTimecodeToSource, true));
+
+	// Start up the receiver under the player's control.
+	return Receiver->Initialize(SourceSettings, FNDIStreamReceiver::ECaptureMode::Manual);
+}
+
+void FNDIMediaStreamPlayer::Close()
+{
+	NDIPlayerState = EMediaState::Closed;
+
+	if (Receiver != nullptr)
+	{
+		// Disconnect from receiver events
+		Receiver->OnVideoFrameReceived.Remove(VideoReceivedHandle);
+		VideoReceivedHandle.Reset();
+		Receiver->OnAudioFrameReceived.Remove(AudioReceivedHandle);
+		AudioReceivedHandle.Reset();
+		Receiver->OnConnected.Remove(ConnectedHandle);
+		ConnectedHandle.Reset();
+		Receiver->OnDisconnected.Remove(DisconnectedHandle);
+		DisconnectedHandle.Reset();
+
+		Receiver.Reset();
+	}
+
+	TextureSamplePool->Reset();
+	AudioSamplePool->Reset();
+	MetadataSamplePool->Reset();
+
+#if WITH_EDITOR
+	if (!bIsReopening)
+	{
+		OptionsObject = nullptr;
+		UNDIMediaSource::OnOptionChanged.RemoveAll(this);
+	}
+#endif
+
+	Super::Close();
+}
+
+FString FNDIMediaStreamPlayer::GetStats() const
+{
+	FString Stats;
+
+	if (Receiver)
+	{
+		const FNDIMediaReceiverPerformanceData PerformanceData = Receiver->GetPerformanceData();
+		Stats += FString::Printf(TEXT("Video Frames: %lld"), PerformanceData.VideoFrames);
+		Stats += FString::Printf(TEXT("Dropped Video Frames: %lld"), PerformanceData.DroppedVideoFrames);
+		Stats += FString::Printf(TEXT("Audio Frames: %lld"), PerformanceData.AudioFrames);
+		Stats += FString::Printf(TEXT("Dropped Audio Frames: %lld"), PerformanceData.DroppedAudioFrames);
+		Stats += FString::Printf(TEXT("Metadata Frames: %lld"), PerformanceData.MetadataFrames);
+		Stats += FString::Printf(TEXT("Dropped Metadata Frames: %lld"), PerformanceData.DroppedMetadataFrames);
+	}
+	else
+	{
+		Stats = FString(TEXT("Receiver not available."));
+	}
+	return Stats;
+}
+
+void FNDIMediaStreamPlayer::TickInput(FTimespan InDeltaTime, FTimespan InTime)
+{
+	// Update player state
+	EMediaState NewState = NDIPlayerState;
+
+	if (NewState != CurrentState)
+	{
+		CurrentState = NewState;
+		if (CurrentState == EMediaState::Playing)
+		{
+			EventSink.ReceiveMediaEvent(EMediaEvent::TracksChanged);
+			EventSink.ReceiveMediaEvent(EMediaEvent::MediaOpened);
+			EventSink.ReceiveMediaEvent(EMediaEvent::PlaybackResumed);
+		}
+		else if (NewState == EMediaState::Error)
+		{
+			EventSink.ReceiveMediaEvent(EMediaEvent::MediaOpenFailed);
+			Close();
+		}
+	}
+
+	if (CurrentState != EMediaState::Playing)
+	{
+		return;
+	}
+
+	TickTimeManagement();
+}
+
+
+void FNDIMediaStreamPlayer::TickFetch(FTimespan InDeltaTime, FTimespan InTime)
+{
+	Super::TickFetch(InDeltaTime, InTime);
+
+	if (CurrentState == EMediaState::Preparing || CurrentState == EMediaState::Playing)
+	{
+		if (Receiver != nullptr)
+		{
+			if (bCaptureAudio)
+			{
+				Receiver->FetchAudio(InTime);
+			}
+			if (bCaptureVideo)
+			{
+				Receiver->FetchVideo(InTime);
+			}
+			if (bCaptureAncillary)
+			{
+				// Potential improvement: limit how much metadata is processed, to avoid appearing to lock up due to a metadata flood
+				while (Receiver->FetchMetadata(InTime)) {}
+			}
+		}
+	}
+
+	if (CurrentState == EMediaState::Playing)
+	{
+		// No need to lock here. That info is only used for debug information.
+		AudioTrackFormat.NumChannels = NDIThreadAudioChannels;
+		AudioTrackFormat.SampleRate = NDIThreadAudioSampleRate;
+
+		if (Receiver)
+		{
+			VideoFrameRate = Receiver->GetCurrentFrameRate();
+			VideoTrackFormat.Dim = Receiver->GetCurrentResolution();
+			VideoTrackFormat.FrameRates = TRange<float>(VideoFrameRate.AsDecimal());
+			VideoTrackFormat.FrameRate = VideoFrameRate.AsDecimal();
+			//VideoTrackFormat.TypeName = Configuration.Configuration.MediaMode.GetModeName().ToString();
+		}
+
+		VerifyFrameDropCount();
+	}
+}
+
+void FNDIMediaStreamPlayer::HandleVideoFrameReceived(FNDIStreamReceiver* InReceiver, const NDIlib_video_frame_v2_t& InVideoFrame, const FTimespan& InTime)
+{
+	TSharedRef<FNDIMediaTextureSample> TextureSample = TextureSamplePool->AcquireShared();
+	
+	const UE::MediaIOCore::FColorFormatArgs ColorFormatArgs(
+		bOverrideSourceEncoding ?  static_cast<UE::Color::EEncoding>(OverrideSourceEncoding) : UE::Color::EEncoding::sRGB,
+		bOverrideSourceColorSpace ? static_cast<UE::Color::EColorSpace>(OverrideSourceColorSpace) : UE::Color::EColorSpace::sRGB);
+
+	FTimecode SourceTimecode = InReceiver->GetCurrentTimecode();
+
+	if (TextureSample->Initialize(InVideoFrame, ColorFormatArgs, InTime, SourceTimecode))
+	{
+		if (TextureSample->CustomConverter)
+		{
+			TextureSample->CustomConverter->Setup(TextureSample);
+		}
+		TextureSample->SetColorConversionSettings(OCIOSettings);
+
+		if (bEncodeTimecodeInTexel && InVideoFrame.frame_format_type == NDIlib_frame_format_type_progressive)
+		{
+			EMediaIOCoreEncodePixelFormat EncodePixelFormat;
+			bool bEncodeSupported = true;
+
+			if (InVideoFrame.FourCC == NDIlib_FourCC_video_type_UYVY
+				|| InVideoFrame.FourCC == NDIlib_FourCC_video_type_UYVA)
+			{
+				// Note: for UYVA, we can write in the UYVY part (even if it ends up being transparent).
+				// todo: add support in FMediaIOCoreEncodeTime for single channel (R) format.
+				EncodePixelFormat = EMediaIOCoreEncodePixelFormat::CharUYVY;
+			}
+			else if (InVideoFrame.FourCC == NDIlib_FourCC_video_type_BGRA
+				|| InVideoFrame.FourCC == NDIlib_FourCC_video_type_RGBA
+				|| InVideoFrame.FourCC == NDIlib_FourCC_video_type_BGRX)
+			{
+				EncodePixelFormat = EMediaIOCoreEncodePixelFormat::CharBGRA;
+			}
+			else
+			{
+				EncodePixelFormat = EMediaIOCoreEncodePixelFormat::CharUYVY;
+				bEncodeSupported = false;
+			}
+			
+			if (bEncodeSupported)
+			{
+				FMediaIOCoreEncodeTime EncodeTime(EncodePixelFormat, const_cast<void*>(TextureSample->GetBuffer()), InVideoFrame.line_stride_in_bytes, InVideoFrame.xres, InVideoFrame.yres);
+				EncodeTime.Render(SourceTimecode.Hours, SourceTimecode.Minutes, SourceTimecode.Seconds, SourceTimecode.Frames);
+			}
+		}
+		
+		AddVideoSample(TextureSample);
+	}
+}
+
+void FNDIMediaStreamPlayer::HandleAudioFrameReceived(FNDIStreamReceiver* InReceiver, const NDIlib_audio_frame_v2_t& InAudioFrame, const FTimespan& InTime)
+{
+	FNDIMediaRuntimeLibrary* NdiLib = Receiver->GetNdiLib();
+	if (!NdiLib || !NdiLib->IsLoaded())
+	{
+		return;
+	}
+	
+	TSharedRef<FNDIMediaAudioSample> AudioSample = AudioSamplePool->AcquireShared();
+
+	// UE wants 32bit signed interleaved audio data, so need to convert the NDI audio.
+	// Fortunately the NDI library has a utility function to do that.
+
+	// Get a buffer to convert to
+	const int32 available_samples = InAudioFrame.no_samples * InAudioFrame.no_channels;
+	void* SampleBuffer = AudioSample->RequestBuffer(available_samples);
+
+	if (SampleBuffer != nullptr)
+	{
+		// Format to convert to
+		NDIlib_audio_frame_interleaved_32s_t audio_frame_32s(
+			InAudioFrame.sample_rate,
+			InAudioFrame.no_channels,
+			InAudioFrame.no_samples,
+			InAudioFrame.timecode,
+			20,
+			static_cast<int32_t*>(SampleBuffer));
+
+		// Convert received NDI audio
+		NdiLib->Lib->util_audio_to_interleaved_32s_v2(&InAudioFrame, &audio_frame_32s);
+		
+		// Supply converted audio data
+		if (AudioSample->SetProperties(available_samples
+			, audio_frame_32s.no_channels
+			, audio_frame_32s.sample_rate
+			, InTime
+			, TOptional<FTimecode>()))
+		{
+			NDIThreadAudioChannels = audio_frame_32s.no_channels;
+			NDIThreadAudioSampleRate = audio_frame_32s.sample_rate;
+			
+			AddAudioSample(AudioSample);
+		}
+	}
+}
+
+void FNDIMediaStreamPlayer::VerifyFrameDropCount()
+{
+	// todo
+}
+
+bool FNDIMediaStreamPlayer::IsHardwareReady() const
+{
+	return NDIPlayerState == EMediaState::Playing;
+}
+
+void FNDIMediaStreamPlayer::SetupSampleChannels()
+{
+	FMediaIOSamplingSettings VideoSettings = BaseSettings;
+	VideoSettings.BufferSize = MaxNumVideoFrameBuffer;
+	Samples->InitializeVideoBuffer(VideoSettings);
+
+	FMediaIOSamplingSettings AudioSettings = BaseSettings;
+	AudioSettings.BufferSize = MaxNumAudioFrameBuffer;
+	Samples->InitializeAudioBuffer(AudioSettings);
+
+	FMediaIOSamplingSettings MetadataSettings = BaseSettings;
+	MetadataSettings.BufferSize = MaxNumMetadataFrameBuffer;
+	Samples->InitializeMetadataBuffer(MetadataSettings);
+}
+
+
+TSharedPtr<FMediaIOCoreTextureSampleBase> FNDIMediaStreamPlayer::AcquireTextureSample_AnyThread() const
+{
+	// Needed by the deinterlacer.
+	return TextureSamplePool->AcquireShared();
+}
+
+//~ ITimedDataInput interface
+#if WITH_EDITOR
+const FSlateBrush* FNDIMediaStreamPlayer::GetDisplayIcon() const
+{
+	return nullptr;
+}
+#endif
+
+
+#undef LOCTEXT_NAMESPACE
diff --git a/Engine/Plugins/Media/NDIMedia/Source/NDIMedia/Private/Player/NDIMediaStreamPlayer.h b/Engine/Plugins/Media/NDIMedia/Source/NDIMedia/Private/Player/NDIMediaStreamPlayer.h
new file mode 100644
--- /dev/null
+++ b/Engine/Plugins/Media/NDIMedia/Source/NDIMedia/Private/Player/NDIMediaStreamPlayer.h
@@ -0,0 +1,113 @@
+// Copyright Epic Games, Inc. All Rights Reserved.
+
+#pragma once
+
+#include "MediaIOCorePlayerBase.h"
+#include "NDIMediaAPI.h"
+#include "UObject/ObjectKey.h"
+
+class FNDIMediaAudioSamplePool;
+class FNDIMediaBinarySamplePool;
+class FNDIMediaRuntimeLibrary;
+class FNDIMediaTextureSamplePool;
+class FNDIStreamReceiver;
+
+/**
+ * Implementation of the Media player for an NDI stream.
+ */
+class FNDIMediaStreamPlayer : public FMediaIOCorePlayerBase
+{
+	using Super = FMediaIOCorePlayerBase;
+
+public:
+	FNDIMediaStreamPlayer(IMediaEventSink& InEventSink);
+
+	virtual ~FNDIMediaStreamPlayer() override;
+
+	//~ Begin IMediaPlayer
+	virtual FGuid GetPlayerPluginGUID() const override;
+	virtual bool Open(const FString& InUrl, const IMediaOptions* InOptions) override;
+	virtual void Close() override;
+	virtual FString GetStats() const override;
+	virtual void TickInput(FTimespan InDeltaTime, FTimespan InTime) override;
+	virtual void TickFetch(FTimespan InDeltaTime, FTimespan InTime) override;
+
+protected:
+	virtual bool IsHardwareReady() const override;
+	virtual void SetupSampleChannels() override;
+	virtual TSharedPtr<FMediaIOCoreTextureSampleBase> AcquireTextureSample_AnyThread() const override;
+	//~ End IMediaPlayer
+
+public:
+	//~ Begin ITimedDataInput
+#if WITH_EDITOR
+	virtual const FSlateBrush* GetDisplayIcon() const override;
+#endif
+	//~ End ITimedDataInput
+
+private:
+#if WITH_EDITOR
+	void OnOptionsChanged(UObject* InOptions, FPropertyChangedEvent& InPropertyChanged);
+#endif
+
+	/** Callback handler for video frame reception. */
+	void HandleVideoFrameReceived(FNDIStreamReceiver* InReceiver, const NDIlib_video_frame_v2_t& InVideoFrame, const FTimespan& InTime);
+
+	/** Callback handler for audio frame reception. */
+	void HandleAudioFrameReceived(FNDIStreamReceiver* InReceiver, const NDIlib_audio_frame_v2_t& InAudioFrame, const FTimespan& InTime);
+
+	void VerifyFrameDropCount();
+
+private:
+	/** Keep track of the options (media source) object this player was opened with to filter global events. */
+	FObjectKey OptionsObject;
+
+	/** Indicate if currently in a reopening sequence. */
+	bool bIsReopening = false;
+	
+	/** Maximum Audio Sample Pool size */
+	int32 MaxNumAudioFrameBuffer = 0;
+	
+	/** Maximum Metadata Sample Pool size */
+	int32 MaxNumMetadataFrameBuffer = 0;
+
+	/** Maximum Video Sample Pool size */
+	int32 MaxNumVideoFrameBuffer = 0;
+
+	/** Determines which of the sample streams to capture (mirrors MediaSource). */
+	bool bCaptureVideo = true;
+	bool bCaptureAudio = false;
+	bool bCaptureAncillary = false;
+
+	/** Whether to use the time code embedded in video frames. */
+	bool bEncodeTimecodeInTexel = false;
+
+	/** Used to flag which sample types we advertise as supported for timed data monitoring */
+	EMediaIOSampleType SupportedSampleTypes = EMediaIOSampleType::None;
+
+	/** Current state of the media player. */
+	EMediaState NDIPlayerState = EMediaState::Closed;
+
+	/** Number of channels from the last audio frame received. */
+	uint32 NDIThreadAudioChannels = 0;
+
+	/** Sample rate from the last audio frame received. */
+	uint32 NDIThreadAudioSampleRate = 0;
+	
+	/** The media event sink. */
+	IMediaEventSink& EventSink;
+
+	/** Current stream receiver */
+	TSharedPtr<FNDIStreamReceiver> Receiver;
+
+	/** Handles for the receiver delegates */ 
+	FDelegateHandle VideoReceivedHandle;
+	FDelegateHandle AudioReceivedHandle;
+	FDelegateHandle ConnectedHandle;
+	FDelegateHandle DisconnectedHandle;
+
+	/** Media Sample pools */
+	TUniquePtr<FNDIMediaTextureSamplePool> TextureSamplePool;
+	TUniquePtr<FNDIMediaAudioSamplePool> AudioSamplePool;
+	TUniquePtr<FNDIMediaBinarySamplePool> MetadataSamplePool;
+};
diff --git a/Engine/Plugins/Media/NDIMedia/Source/NDIMedia/Private/Player/NDIReceiverPerformanceData.h b/Engine/Plugins/Media/NDIMedia/Source/NDIMedia/Private/Player/NDIReceiverPerformanceData.h
deleted file mode 100644
diff --git a/Engine/Plugins/Media/NDIMedia/Source/NDIMedia/Private/Player/NDIStreamReceiver.cpp b/Engine/Plugins/Media/NDIMedia/Source/NDIMedia/Private/Player/NDIStreamReceiver.cpp
--- a/Engine/Plugins/Media/NDIMedia/Source/NDIMedia/Private/Player/NDIStreamReceiver.cpp
+++ b/Engine/Plugins/Media/NDIMedia/Source/NDIMedia/Private/Player/NDIStreamReceiver.cpp
@@ -325,7 +325,7 @@ void FNDIStreamReceiver::Shutdown()
 	SetIsCurrentlyConnected(/*bInConnected*/ false, /*bInDelayBroadcastEvents*/ false);
 
 	SourceSettings = FNDISourceSettings();
-	PerformanceData = FNDIReceiverPerformanceData();
+	PerformanceData = FNDIMediaReceiverPerformanceData();
 	FrameRate = FFrameRate(60, 1);
 	Resolution = FIntPoint(0, 0);
 	Timecode = FTimecode(0, FrameRate, true, true);
@@ -580,7 +580,7 @@ void FNDIStreamReceiver::GatherPerformanceMetrics()
 	PerformanceData.VideoFrames = stable_performance.video_frames;
 }
 
-FNDIReceiverPerformanceData FNDIStreamReceiver::GetPerformanceData() const
+FNDIMediaReceiverPerformanceData FNDIStreamReceiver::GetPerformanceData() const
 {
 	FScopeLock Lock(&RenderSyncContext);
 	return PerformanceData;
diff --git a/Engine/Plugins/Media/NDIMedia/Source/NDIMedia/Private/Player/NDIStreamReceiver.h b/Engine/Plugins/Media/NDIMedia/Source/NDIMedia/Private/Player/NDIStreamReceiver.h
--- a/Engine/Plugins/Media/NDIMedia/Source/NDIMedia/Private/Player/NDIStreamReceiver.h
+++ b/Engine/Plugins/Media/NDIMedia/Source/NDIMedia/Private/Player/NDIStreamReceiver.h
@@ -6,7 +6,7 @@
 #include "Misc/Timecode.h"
 #include "NDISourceSettings.h"
 #include "NDIMediaAPI.h"
-#include "NDIReceiverPerformanceData.h"
+#include "NDIMediaReceiverPerformanceData.h"
 #include "Templates/SharedPointer.h"
 
 class FNDIMediaRuntimeLibrary;
@@ -102,7 +102,7 @@ public:
 	const FNDISourceSettings& GetCurrentSourceSettings() const;
 
 	/** Returns the current performance data of the receiver while connected to the source */
-	FNDIReceiverPerformanceData GetPerformanceData() const;
+	FNDIMediaReceiverPerformanceData GetPerformanceData() const;
 
 	/** Returns a value indicating whether this object is currently connected to the sender source */
 	bool GetIsCurrentlyConnected() const;
@@ -174,7 +174,7 @@ private:
 	FNDISourceSettings SourceSettings;
 
 	/** Information describing detailed information about the receiver performance when connected to an NDI? sender */
-	FNDIReceiverPerformanceData PerformanceData;
+	FNDIMediaReceiverPerformanceData PerformanceData;
 
 	/** Keep track of the last VideoFrame timestamp received. Used to detect new frame. */
 	int64_t LastFrameTimestamp = 0;
diff --git a/Engine/Plugins/Media/NDIMedia/Source/NDIMedia/Public/NDIMediaTimecodeProvider.h b/Engine/Plugins/Media/NDIMedia/Source/NDIMedia/Public/NDIMediaTimecodeProvider.h
new file mode 100644
--- /dev/null
+++ b/Engine/Plugins/Media/NDIMedia/Source/NDIMedia/Public/NDIMediaTimecodeProvider.h
@@ -0,0 +1,63 @@
+// Copyright Epic Games, Inc. All Rights Reserved.
+
+#pragma once
+
+#include "GenlockedTimecodeProvider.h"
+#include "MediaIOCoreDefinitions.h"
+#include "NDIMediaDefines.h"
+
+#include "NDIMediaTimecodeProvider.generated.h"
+
+class FNDIStreamReceiver;
+
+/**
+ * Timecode provider from an NDI source 
+ */
+UCLASS(MinimalAPI, Blueprintable, editinlinenew, meta=(DisplayName="NDI Timecode Provider", MediaIOCustomLayout="NDI"))
+class UNDIMediaTimecodeProvider : public UGenlockedTimecodeProvider
+{
+	GENERATED_BODY()
+
+public:
+	//~ Begin UTimecodeProvider
+	NDIMEDIA_API virtual bool FetchTimecode(FQualifiedFrameTime& OutFrameTime) override;
+	NDIMEDIA_API virtual ETimecodeProviderSynchronizationState GetSynchronizationState() const override;
+	NDIMEDIA_API virtual bool Initialize(class UEngine* InEngine) override;
+	NDIMEDIA_API virtual void Shutdown(class UEngine* InEngine) override;
+	//~ End UTimecodeProvider
+
+	//~ Begin UObject
+	NDIMEDIA_API virtual void BeginDestroy() override;
+	//~ End UObject
+
+	/**
+	 * Use the time code embedded in the video input stream. 
+	 */
+	UPROPERTY(EditAnywhere, Category="Timecode")
+	FMediaIOVideoTimecodeConfiguration TimecodeConfiguration;
+
+	/** Indicates the current bandwidth mode used for the connection to this source */
+	UPROPERTY(EditAnywhere, Category="NDI")
+	ENDIReceiverBandwidth Bandwidth = ENDIReceiverBandwidth::Highest;
+
+private:
+	void ReleaseResources();
+
+private:
+	/** Critical section protecting "State" */
+	mutable FCriticalSection StateSyncContext;
+
+	/** Current synchronization state */
+	ETimecodeProviderSynchronizationState State = ETimecodeProviderSynchronizationState::Closed;
+
+	/** Last received frame time from the receiver. */
+	FQualifiedFrameTime MostRecentFrameTime;
+
+	/** Current stream receiver */
+	TSharedPtr<FNDIStreamReceiver> Receiver;
+
+	/** Handles for the stream receiver delegates */
+	FDelegateHandle VideoFrameReceivedHandle;
+	FDelegateHandle ConnectedHandle;
+	FDelegateHandle DisconnectedHandle;
+};
diff --git a/Engine/Plugins/Media/NDIMedia/Source/NDIMedia/Public/NDITimecodeProvider.h b/Engine/Plugins/Media/NDIMedia/Source/NDIMedia/Public/NDITimecodeProvider.h
deleted file mode 100644
