[NDI Media] Updating description to mention media source.

diff --git a/Engine/Plugins/Media/NDIMedia/NDIMedia.uplugin b/Engine/Plugins/Media/NDIMedia/NDIMedia.uplugin
--- a/Engine/Plugins/Media/NDIMedia/NDIMedia.uplugin
+++ b/Engine/Plugins/Media/NDIMedia/NDIMedia.uplugin
@@ -3,7 +3,7 @@
 	"Version": 1,
 	"VersionName": "1.0",
 	"FriendlyName": "NDI Media",
-	"Description": "Implements media output using NDI protocol",
+	"Description": "Implements media source and media output using NDI protocol",
 	"Category": "Media Players",
 	"CreatedBy": "Epic Games",
 	"CreatedByURL": "http://epicgames.com",
