[NDI Media] Implement "any thread" capture api (latency improvement)

diff --git a/Engine/Plugins/Media/NDIMedia/Source/NDIMedia/Private/NDIMediaCapture.cpp b/Engine/Plugins/Media/NDIMedia/Source/NDIMedia/Private/NDIMediaCapture.cpp
--- a/Engine/Plugins/Media/NDIMedia/Source/NDIMedia/Private/NDIMediaCapture.cpp
+++ b/Engine/Plugins/Media/NDIMedia/Source/NDIMedia/Private/NDIMediaCapture.cpp
@@ -319,7 +319,17 @@ inline int64_t ConvertToNDITimeCode(const FTimecode& InTimecode, const FFrameRat
 void UNDIMediaCapture::OnFrameCaptured_RenderingThread(const FCaptureBaseData& InBaseData, TSharedPtr<FMediaCaptureUserData, ESPMode::ThreadSafe> InUserData, void* InBuffer, int32 Width, int32 Height, int32 BytesPerRow)
 {
 	TRACE_CPUPROFILER_EVENT_SCOPE(UNDIMediaCapture::OnFrameCaptured_RenderingThread);
+	OnFrameCapturedImpl(InBaseData, InBuffer, Width, Height, BytesPerRow);
+}
 
+void UNDIMediaCapture::OnFrameCaptured_AnyThread(const FCaptureBaseData& InBaseData, TSharedPtr<FMediaCaptureUserData, ESPMode::ThreadSafe> InUserData,  const FMediaCaptureResourceData& InResourceData)
+{
+	TRACE_CPUPROFILER_EVENT_SCOPE(UNDIMediaCapture::OnFrameCaptured_AnyThread);
+	OnFrameCapturedImpl(InBaseData, InResourceData.Buffer, InResourceData.Width, InResourceData.Height, InResourceData.BytesPerRow);
+}
+
+void UNDIMediaCapture::OnFrameCapturedImpl(const FCaptureBaseData& InBaseData, void* InBuffer, int32 Width, int32 Height, int32 BytesPerRow)
+{
 	FScopeLock ScopeLock(&CaptureInstanceCriticalSection);
 
 	if (CaptureInstance && CaptureInstance->Sender)
@@ -437,6 +447,11 @@ void UNDIMediaCapture::StopCaptureImpl(bool /*bAllowPendingFrameToBeProcess*/)
 	CaptureInstance = nullptr;
 }
 
+bool UNDIMediaCapture::SupportsAnyThreadCapture() const
+{
+	return true;
+}
+
 bool UNDIMediaCapture::StartNewCapture()
 {
 	TRACE_BOOKMARK(TEXT("NDIMediaCapture::StartNewCapture"));
diff --git a/Engine/Plugins/Media/NDIMedia/Source/NDIMedia/Public/NDIMediaCapture.h b/Engine/Plugins/Media/NDIMedia/Source/NDIMedia/Public/NDIMediaCapture.h
--- a/Engine/Plugins/Media/NDIMedia/Source/NDIMedia/Public/NDIMediaCapture.h
+++ b/Engine/Plugins/Media/NDIMedia/Source/NDIMedia/Public/NDIMediaCapture.h
@@ -19,17 +19,24 @@ public:
 	virtual ~UNDIMediaCapture();
 
 protected:
+	//~ Begin UMediaCapture
 	virtual void OnFrameCaptured_RenderingThread(const FCaptureBaseData& InBaseData, TSharedPtr<FMediaCaptureUserData, ESPMode::ThreadSafe> InUserData, void* InBuffer, int32 Width, int32 Height, int32 BytesPerRow) override;
+	virtual void OnFrameCaptured_AnyThread(const FCaptureBaseData& InBaseData, TSharedPtr<FMediaCaptureUserData, ESPMode::ThreadSafe> InUserData, const FMediaCaptureResourceData& InResourceData) override;
 	virtual bool InitializeCapture() override;
 	virtual bool PostInitializeCaptureViewport(TSharedPtr<FSceneViewport>& InSceneViewport) override;
 	virtual bool PostInitializeCaptureRenderTarget(UTextureRenderTarget2D* InRenderTarget) override;
 	virtual bool PostInitializeCaptureRHIResource(const FRHICaptureResourceDescription& InResourceDescription) override;
 	virtual bool UpdateAudioDeviceImpl(const FAudioDeviceHandle& InAudioDeviceHandle) override;
 	virtual void StopCaptureImpl(bool bAllowPendingFrameToBeProcess) override;
+	virtual bool SupportsAnyThreadCapture() const override;
+	//~ End UMediaCapture
 
 private:
 	bool StartNewCapture();
 
+	/** Sends the captured buffer to NDI. */
+	void OnFrameCapturedImpl(const FCaptureBaseData& InBaseData, void* InBuffer, int32 Width, int32 Height, int32 BytesPerRow);
+
 private:
 	FCriticalSection CaptureInstanceCriticalSection;
 
