[NDI Media] Crash fix for NDIMediaOutput on Mac Platform - SupportsAnyThreadCapture is not supported, it must be on the render thread.

diff --git a/Engine/Plugins/Media/NDIMedia/Source/NDIMedia/Private/NDIMediaCapture.cpp b/Engine/Plugins/Media/NDIMedia/Source/NDIMedia/Private/NDIMediaCapture.cpp
--- a/Engine/Plugins/Media/NDIMedia/Source/NDIMedia/Private/NDIMediaCapture.cpp
+++ b/Engine/Plugins/Media/NDIMedia/Source/NDIMedia/Private/NDIMediaCapture.cpp
@@ -452,7 +452,14 @@ void UNDIMediaCapture::StopCaptureImpl(bool /*bAllowPendingFrameToBeProcess*/)
 
 bool UNDIMediaCapture::SupportsAnyThreadCapture() const
 {
+#if PLATFORM_MAC
+	// On Mac, Media Capture must capture cpu frames on the render thread.
+	// This is because MediaCapture Lock_Unsafe uses RHIMapStagingSurface
+	// which needs to run on the rendering thread with Metal.
+	return false;
+#else  // PLATFORM_MAC
 	return true;
+#endif // PLATFORM_MAC
 }
 
 bool UNDIMediaCapture::StartNewCapture()
