[NDIMedia] Fixing low quality render when receiving an NDI stream with alpha channel.

The NDI specific UYVYA shader's sampler was set to trilinear. Using point filtering, as is done for the other media shaders, fixes the issue.

diff --git a/Engine/Plugins/Media/NDIMedia/Shaders/Private/NDIMediaShaders.usf b/Engine/Plugins/Media/NDIMedia/Shaders/Private/NDIMediaShaders.usf
--- a/Engine/Plugins/Media/NDIMedia/Shaders/Private/NDIMediaShaders.usf
+++ b/Engine/Plugins/Media/NDIMedia/Shaders/Private/NDIMediaShaders.usf
@@ -82,14 +82,13 @@ void NDIMediaUYVYAConvertPS(
 	out float4 OutColor : SV_Target0)
 {
 	float4 UYVYB = NDIMediaUYVYAConvertUB.YUVTexture.Sample(NDIMediaUYVYAConvertUB.SamplerB, InUV);
-	float4 UYVYT = NDIMediaUYVYAConvertUB.YUVTexture.Sample(NDIMediaUYVYAConvertUB.SamplerT, InUV);
-	float Alpha = NDIMediaUYVYAConvertUB.AlphaTexture.Sample(NDIMediaUYVYAConvertUB.SamplerT, InUV).w;
+	float4 UYVYP = NDIMediaUYVYAConvertUB.YUVTexture.Sample(NDIMediaUYVYAConvertUB.SamplerP, InUV);
+	float Alpha = NDIMediaUYVYAConvertUB.AlphaTexture.Sample(NDIMediaUYVYAConvertUB.SamplerP, InUV).w;
 
 	float PosX = 2.0f * InUV.x * NDIMediaUYVYAConvertUB.InputWidth;
-
-	float FracX = PosX % 2.0f;
 	float3 YUV;
-	YUV.x = (1 - FracX) * UYVYT.y + FracX * UYVYT.w;
+
+	YUV.x = (floor(PosX % 2.0f) == 0) ? UYVYP.y : UYVYP.w;
 	YUV.yz = UYVYB.zx;
 
 	OutColor = float4(ApplyColorProcessing(YUV, NDIMediaUYVYAConvertUB.ColorTransform, NDIMediaUYVYAConvertUB.CSTransform, NDIMediaUYVYAConvertUB.EOTF, NDIMediaUYVYAConvertUB.ToneMapMethod), Alpha);
diff --git a/Engine/Plugins/Media/NDIMedia/Source/NDIMediaRendering/Private/NDIMediaShaders.cpp b/Engine/Plugins/Media/NDIMedia/Source/NDIMediaRendering/Private/NDIMediaShaders.cpp
--- a/Engine/Plugins/Media/NDIMedia/Source/NDIMediaRendering/Private/NDIMediaShaders.cpp
+++ b/Engine/Plugins/Media/NDIMedia/Source/NDIMediaRendering/Private/NDIMediaShaders.cpp
@@ -16,7 +16,7 @@ BEGIN_GLOBAL_SHADER_PARAMETER_STRUCT(FNDIMediaUYVYAConvertUB, )
 	SHADER_PARAMETER_TEXTURE(Texture2D, YUVTexture)
 	SHADER_PARAMETER_TEXTURE(Texture2D, AlphaTexture)
 	SHADER_PARAMETER_SAMPLER(SamplerState, SamplerB)
-	SHADER_PARAMETER_SAMPLER(SamplerState, SamplerT)
+	SHADER_PARAMETER_SAMPLER(SamplerState, SamplerP)
 END_GLOBAL_SHADER_PARAMETER_STRUCT()
 
 IMPLEMENT_GLOBAL_SHADER_PARAMETER_STRUCT(FNDIMediaUYVYAConvertUB, "NDIMediaUYVYAConvertUB");
@@ -38,7 +38,7 @@ void FNDIMediaShaderUYVAtoBGRAPS::SetParameters(FRHIBatchedShaderParameters& InB
 		UB.YUVTexture = InParameters.YUVTexture;
 		UB.AlphaTexture = InParameters.AlphaTexture;
 		UB.SamplerB = TStaticSamplerState<SF_Bilinear>::GetRHI();
-		UB.SamplerT = TStaticSamplerState<SF_Trilinear>::GetRHI();
+		UB.SamplerP = TStaticSamplerState<SF_Point>::GetRHI();
 	}
 
 	TUniformBufferRef<FNDIMediaUYVYAConvertUB> Data = TUniformBufferRef<FNDIMediaUYVYAConvertUB>::CreateUniformBufferImmediate(UB, UniformBuffer_SingleFrame);
