[NDIMedia] Fix Framerate property to be editable in media profile.

diff --git a/Engine/Plugins/Media/NDIMedia/Source/NDIMedia/Public/NDIMediaOutput.h b/Engine/Plugins/Media/NDIMedia/Source/NDIMedia/Public/NDIMediaOutput.h
--- a/Engine/Plugins/Media/NDIMedia/Source/NDIMedia/Public/NDIMediaOutput.h
+++ b/Engine/Plugins/Media/NDIMedia/Source/NDIMedia/Public/NDIMediaOutput.h
@@ -96,7 +96,7 @@ public:
 	ENDIMediaOutputPixelFormat DesiredPixelFormat = ENDIMediaOutputPixelFormat::B8G8R8A8;
 
 	/** Represents the desired number of frames (per second) for video to be sent over NDI */
-	UPROPERTY(BlueprintReadwrite, EditDefaultsOnly, Category = "Media")
+	UPROPERTY(BlueprintReadwrite, EditAnywhere, Category = "Media")
 	FFrameRate FrameRate = FFrameRate(60, 1);
 	
 	/** Whether to capture and output audio from the engine. */
