[NDIMedia] Fix Just in Time Rendering (JITR) and timecode synchronization.

diff --git a/Engine/Plugins/Media/NDIMedia/Source/NDIMedia/Private/Player/NDIMediaStreamPlayer.cpp b/Engine/Plugins/Media/NDIMedia/Source/NDIMedia/Private/Player/NDIMediaStreamPlayer.cpp
--- a/Engine/Plugins/Media/NDIMedia/Source/NDIMedia/Private/Player/NDIMediaStreamPlayer.cpp
+++ b/Engine/Plugins/Media/NDIMedia/Source/NDIMedia/Private/Player/NDIMediaStreamPlayer.cpp
@@ -9,6 +9,7 @@
 #include "MediaIOCoreEncodeTime.h"
 #include "MediaIOCoreSamples.h"
 #include "MediaIOCoreTextureSampleBase.h"
+#include "MediaIOCoreTextureSampleConverter.h"
 #include "NDIMediaLog.h"
 #include "NDIMediaModule.h"
 #include "NDIMediaSource.h"
@@ -113,6 +114,9 @@ bool FNDIMediaStreamPlayer::Open(const FString& InUrl, const IMediaOptions* InOp
 	AudioTrackFormat.SampleRate = 44100;
 	AudioTrackFormat.TypeName = FString(TEXT("PCM"));
 
+	// Video track format is only known when the first video frame is received.
+	VideoTrackFormat.Dim = FIntPoint(0, 0);
+
 	bCaptureVideo = InOptions->GetMediaOption(UE::NDIMediaSourceOptions::CaptureVideo, true);
 	bCaptureAudio = InOptions->GetMediaOption(UE::NDIMediaSourceOptions::CaptureAudio, false);
 	bCaptureAncillary = InOptions->GetMediaOption(UE::NDIMediaSourceOptions::CaptureAncillary, false);
@@ -260,6 +264,18 @@ void FNDIMediaStreamPlayer::TickInput(FTimespan InDeltaTime, FTimespan InTime)
 	TickTimeManagement();
 }
 
+bool FNDIMediaStreamPlayer::HasVideoSamplesWithCustomConverter() const
+{
+	const TArray<TSharedPtr<IMediaTextureSample>> TextureSamples = Samples->GetVideoSamples();
+	// Check the latest received sample only. New samples are inserted at 0.
+	if (!TextureSamples.IsEmpty() && TextureSamples[0])
+	{
+		// We have to explicitly check the CustomConverter of the NDI sample. 
+		// We can't use GetMediaTextureSampleConverter() because it may return the JITR converter.
+		return static_cast<const FNDIMediaTextureSample*>(TextureSamples[0].Get())->CustomConverter != nullptr;
+	}
+	return false;
+}
 
 void FNDIMediaStreamPlayer::TickFetch(FTimespan InDeltaTime, FTimespan InTime)
 {
@@ -294,10 +310,23 @@ void FNDIMediaStreamPlayer::TickFetch(FTimespan InDeltaTime, FTimespan InTime)
 		if (Receiver)
 		{
 			VideoFrameRate = Receiver->GetCurrentFrameRate();
-			VideoTrackFormat.Dim = Receiver->GetCurrentResolution();
 			VideoTrackFormat.FrameRates = TRange<float>(VideoFrameRate.AsDecimal());
 			VideoTrackFormat.FrameRate = VideoFrameRate.AsDecimal();
 			//VideoTrackFormat.TypeName = Configuration.Configuration.MediaMode.GetModeName().ToString();
+
+			// Detect changes in video format.
+			if (VideoTrackFormat.Dim != Receiver->GetCurrentResolution())
+			{
+				VideoTrackFormat.Dim = Receiver->GetCurrentResolution();
+				NotifyVideoFormatDetected();
+			}
+
+			// Detect changes in frame rate settings for timed data channels.
+			if (BaseSettings.FrameRate != VideoFrameRate)
+			{
+				BaseSettings.FrameRate = VideoFrameRate;
+				SetupSampleChannels();
+			}
 		}
 
 		VerifyFrameDropCount();
@@ -306,6 +335,11 @@ void FNDIMediaStreamPlayer::TickFetch(FTimespan InDeltaTime, FTimespan InTime)
 
 void FNDIMediaStreamPlayer::HandleVideoFrameReceived(FNDIStreamReceiver* InReceiver, const NDIlib_video_frame_v2_t& InVideoFrame, const FTimespan& InTime)
 {
+	if (!InReceiver)
+	{
+		return;
+	}
+
 	TSharedRef<FNDIMediaTextureSample> TextureSample = TextureSamplePool->AcquireShared();
 	
 	const UE::MediaIOCore::FColorFormatArgs ColorFormatArgs(
@@ -314,12 +348,20 @@ void FNDIMediaStreamPlayer::HandleVideoFrameReceived(FNDIStreamReceiver* InRecei
 
 	FTimecode SourceTimecode = InReceiver->GetCurrentTimecode();
 
-	if (TextureSample->Initialize(InVideoFrame, ColorFormatArgs, InTime, SourceTimecode))
+	FTimespan SampleTime = InTime;
+
+	if (EvaluationType == EMediaIOSampleEvaluationType::Timecode)
+	{
+		SampleTime = SourceTimecode.ToTimespan(InReceiver->GetCurrentFrameRate());
+	}
+
+	if (TextureSample->Initialize(InVideoFrame, ColorFormatArgs, SampleTime, SourceTimecode))
 	{
 		if (TextureSample->CustomConverter)
 		{
 			TextureSample->CustomConverter->Setup(TextureSample);
 		}
+
 		TextureSample->SetColorConversionSettings(OCIOSettings);
 
 		if (bEncodeTimecodeInTexel && InVideoFrame.frame_format_type == NDIlib_frame_format_type_progressive)
@@ -428,13 +470,88 @@ void FNDIMediaStreamPlayer::SetupSampleChannels()
 	Samples->InitializeMetadataBuffer(MetadataSettings);
 }
 
+uint32 FNDIMediaStreamPlayer::GetNumVideoFrameBuffers() const
+{
+	return MaxNumVideoFrameBuffer;
+}
 
 TSharedPtr<FMediaIOCoreTextureSampleBase> FNDIMediaStreamPlayer::AcquireTextureSample_AnyThread() const
 {
-	// Needed by the deinterlacer.
+	// Needed by the deinterlacer and JITR.
 	return TextureSamplePool->AcquireShared();
 }
 
+/**
+ * Implementation of a sample converter that is used in the JITR code path.
+ * It behaves like the base class if the samples don't have a custom converter.
+ * When NDI samples have a custom converter, this wrapping converter will
+ * call it on the Proxy sample to render in the provided destination render target.
+ * In that case, the converter is not a "preprocess only" but rather just a "default" converter.
+ */
+class FNDIMediaTextureSampleConverterWrapper : public FMediaIOCoreTextureSampleConverter
+{
+public:
+	//~ Begin IMediaTextureSampleConverter
+	virtual bool Convert(FRHICommandListImmediate& RHICmdList, FTextureRHIRef& InDstTexture, const FConversionHints& Hints) override
+	{
+		// Base class will set up the proxy sample with the selected sample according to JITR logic.
+		if (!FMediaIOCoreTextureSampleConverter::Convert(RHICmdList, InDstTexture, Hints))
+		{
+			return false;
+		}
+
+		TSharedPtr<FMediaIOCoreTextureSampleBase> ProxySample = JITRProxySample.Pin();
+		if (!ProxySample.IsValid())
+		{
+			return false;
+		}
+
+		// We should then be able to actually convert the sample to the destination using the NDI custom converter from the sample.
+		const FNDIMediaTextureSample* ProxySampleNdi = static_cast<FNDIMediaTextureSample*>(ProxySample.Get());
+
+		if (ProxySampleNdi->bIsCustomFormat)
+		{
+			if (IMediaTextureSampleConverter* CustomConverter = ProxySampleNdi->CustomConverter.Get())
+			{
+				// This can be null if the stream switches format on the fly and we have stale samples with a custom format (ignore them).
+				if (InDstTexture.IsValid())
+				{
+					return CustomConverter->Convert(RHICmdList, InDstTexture, Hints);
+				}
+			}
+			return false;
+		}
+
+		// if it is not a custom format, continue the conversion through media texture pipeline instead.
+		return true;
+	}
+
+	virtual uint32 GetConverterInfoFlags() const override
+	{
+		// We need to check in the player if the samples will have a custom converter, if so
+		// we need to indicate that the converter will do a full conversion by returning the default flags.
+		if (TSharedPtr<FMediaIOCoreTextureSampleBase> ProxySample = JITRProxySample.Pin())
+		{
+			if (TSharedPtr<FMediaIOCorePlayerBase> Player = ProxySample->GetPlayer())
+			{
+				if (static_cast<FNDIMediaStreamPlayer*>(Player.Get())->HasVideoSamplesWithCustomConverter())
+				{
+					// Behave like a default converter (that will actually do the conversion).
+					return ConverterInfoFlags_Default;
+				}
+			}
+		}
+
+		return FMediaIOCoreTextureSampleConverter::GetConverterInfoFlags();
+	}
+	//~ End IMediaTextureSampleConverter
+};
+
+TSharedPtr<FMediaIOCoreTextureSampleConverter> FNDIMediaStreamPlayer::CreateTextureSampleConverter() const
+{
+	return MakeShared<FNDIMediaTextureSampleConverterWrapper>();
+}
+
 //~ ITimedDataInput interface
 #if WITH_EDITOR
 const FSlateBrush* FNDIMediaStreamPlayer::GetDisplayIcon() const
diff --git a/Engine/Plugins/Media/NDIMedia/Source/NDIMedia/Private/Player/NDIMediaStreamPlayer.h b/Engine/Plugins/Media/NDIMedia/Source/NDIMedia/Private/Player/NDIMediaStreamPlayer.h
--- a/Engine/Plugins/Media/NDIMedia/Source/NDIMedia/Private/Player/NDIMediaStreamPlayer.h
+++ b/Engine/Plugins/Media/NDIMedia/Source/NDIMedia/Private/Player/NDIMediaStreamPlayer.h
@@ -31,12 +31,21 @@ public:
 	virtual FString GetStats() const override;
 	virtual void TickInput(FTimespan InDeltaTime, FTimespan InTime) override;
 	virtual void TickFetch(FTimespan InDeltaTime, FTimespan InTime) override;
+	//~ End IMediaPlayer
+
+	/**
+	 * Inspects the pending NDI video samples to determine if they have a custom converter.
+	 */
+	bool HasVideoSamplesWithCustomConverter() const;
 
 protected:
+	//~ Begin FMediaIOCorePlayerBase
 	virtual bool IsHardwareReady() const override;
 	virtual void SetupSampleChannels() override;
+	virtual uint32 GetNumVideoFrameBuffers() const override;
 	virtual TSharedPtr<FMediaIOCoreTextureSampleBase> AcquireTextureSample_AnyThread() const override;
-	//~ End IMediaPlayer
+	virtual TSharedPtr<FMediaIOCoreTextureSampleConverter> CreateTextureSampleConverter() const override;
+	//~ End FMediaIOCorePlayerBase
 
 public:
 	//~ Begin ITimedDataInput
diff --git a/Engine/Plugins/Media/NDIMedia/Source/NDIMedia/Private/Player/NDIMediaTextureSample.cpp b/Engine/Plugins/Media/NDIMedia/Source/NDIMedia/Private/Player/NDIMediaTextureSample.cpp
--- a/Engine/Plugins/Media/NDIMedia/Source/NDIMedia/Private/Player/NDIMediaTextureSample.cpp
+++ b/Engine/Plugins/Media/NDIMedia/Source/NDIMedia/Private/Player/NDIMediaTextureSample.cpp
@@ -11,8 +11,9 @@ bool FNDIMediaTextureSample::Initialize(const NDIlib_video_frame_v2_t& InVideoFr
 	bIsProgressive = true;
 	FieldIndex = 0;
 	
-	int32 FrameBufferSize;
+	int32 FrameBufferSize = 0;
 	EMediaTextureSampleFormat FrameSampleFormat;
+	bool bNeedConverter = false;
 
 	if (InVideoFrame.FourCC == NDIlib_FourCC_video_type_UYVY)
 	{
@@ -33,7 +34,7 @@ bool FNDIMediaTextureSample::Initialize(const NDIlib_video_frame_v2_t& InVideoFr
 	else if (InVideoFrame.FourCC == NDIlib_FourCC_video_type_UYVA)
 	{
 		// UYVA format needs a custom converter.
-		bIsCustomFormat = true;
+		bNeedConverter = true;
 		FrameBufferSize = InVideoFrame.line_stride_in_bytes * InVideoFrame.yres + InVideoFrame.xres*InVideoFrame.yres;
 		FrameSampleFormat = EMediaTextureSampleFormat::CharRGBA; // Resulting texture needs to be RGBA.
 	}
@@ -42,17 +43,15 @@ bool FNDIMediaTextureSample::Initialize(const NDIlib_video_frame_v2_t& InVideoFr
 		return false;
 	}
 
-	// Allocate a custom sample converter if needed.
-	if (!CustomConverter && bIsCustomFormat)
-	{
-		CustomConverter = MakeShared<FNDIMediaTextureSampleConverter>();
-	}
+	// Custom converter is recycled. We keep a backup to avoid it being reset by FreeSample (called by the base class's Initialize).
+	TSharedPtr<FNDIMediaTextureSampleConverter> ConverterBackup = CustomConverter;
+
+	bool bInitSuccess = false;
 
 	switch (InVideoFrame.frame_format_type)
 	{
 		case NDIlib_frame_format_type_progressive:
-			bIsProgressive = true;
-			return Super::Initialize(InVideoFrame.p_data
+			bInitSuccess = Super::Initialize(InVideoFrame.p_data
 				, FrameBufferSize
 				, InVideoFrame.line_stride_in_bytes
 				, InVideoFrame.xres
@@ -62,13 +61,12 @@ bool FNDIMediaTextureSample::Initialize(const NDIlib_video_frame_v2_t& InVideoFr
 				, FFrameRate(InVideoFrame.frame_rate_N, InVideoFrame.frame_rate_D)
 				, InTimecode
 				, InColorFormatArgs);
+			bIsProgressive = true;
+			break;
 
 		case NDIlib_frame_format_type_field_0:
 		case NDIlib_frame_format_type_field_1:
-			bIsProgressive = false;
-			FieldIndex = InVideoFrame.frame_format_type == NDIlib_frame_format_type_field_0 ? 0 : 1;
-		
-		return Super::InitializeWithEvenOddLine(InVideoFrame.frame_format_type == NDIlib_frame_format_type_field_0
+			bInitSuccess = Super::InitializeWithEvenOddLine(InVideoFrame.frame_format_type == NDIlib_frame_format_type_field_0
 			, InVideoFrame.p_data
 			, FrameBufferSize
 			, InVideoFrame.line_stride_in_bytes
@@ -79,10 +77,22 @@ bool FNDIMediaTextureSample::Initialize(const NDIlib_video_frame_v2_t& InVideoFr
 			, FFrameRate(InVideoFrame.frame_rate_N, InVideoFrame.frame_rate_D)
 			, InTimecode
 			, InColorFormatArgs);
-
+			bIsProgressive = false;
+			FieldIndex = InVideoFrame.frame_format_type == NDIlib_frame_format_type_field_0 ? 0 : 1;
+			break;
 		default:
-			return false;
+			break;
+	}
+
+	if (bInitSuccess && bNeedConverter)
+	{
+		bIsCustomFormat = bNeedConverter;
+
+		// Either recycle or allocate a new custom sample converter.
+		CustomConverter = ConverterBackup.IsValid() ? ConverterBackup : MakeShared<FNDIMediaTextureSampleConverter>();
 	}
+
+	return bInitSuccess;
 }
 
 IMediaTextureSampleConverter* FNDIMediaTextureSample::GetMediaTextureSampleConverter()
@@ -93,3 +103,30 @@ IMediaTextureSampleConverter* FNDIMediaTextureSample::GetMediaTextureSampleConve
 	}
 	return Super::GetMediaTextureSampleConverter();
 }
+
+void FNDIMediaTextureSample::CopyConfiguration(const TSharedPtr<FMediaIOCoreTextureSampleBase>& SourceSample)
+{
+	if (FNDIMediaTextureSample* SourceSampleNdi = static_cast<FNDIMediaTextureSample*>(SourceSample.Get()))
+	{
+		bIsCustomFormat = SourceSampleNdi->bIsCustomFormat;
+		bIsProgressive = SourceSampleNdi->bIsProgressive;
+		FieldIndex = SourceSampleNdi->FieldIndex;
+		CustomConverter = SourceSampleNdi->CustomConverter;
+	}
+
+	// We need to preserve the JITR converter so we don't lose it.
+	TSharedPtr<FMediaIOCoreTextureSampleConverter> ProxyConverterBackup = Converter;
+
+	// NDI source samples don't have a JITR converter and will reset the JITR proxy sample's.
+	FMediaIOCoreTextureSampleBase::CopyConfiguration(SourceSample);
+
+	// Restore the JITR converter.
+	Converter = ProxyConverterBackup;
+}
+
+void FNDIMediaTextureSample::FreeSample()
+{
+	bIsCustomFormat = false;
+	CustomConverter.Reset();
+	FMediaIOCoreTextureSampleBase::FreeSample();
+}
\ No newline at end of file
diff --git a/Engine/Plugins/Media/NDIMedia/Source/NDIMedia/Private/Player/NDIMediaTextureSample.h b/Engine/Plugins/Media/NDIMedia/Source/NDIMedia/Private/Player/NDIMediaTextureSample.h
--- a/Engine/Plugins/Media/NDIMedia/Source/NDIMedia/Private/Player/NDIMediaTextureSample.h
+++ b/Engine/Plugins/Media/NDIMedia/Source/NDIMedia/Private/Player/NDIMediaTextureSample.h
@@ -21,6 +21,11 @@ public:
 #endif
 	//~ End IMediaTextureSample
 
+	//~ Begin FMediaIOCoreTextureSampleBase
+	void CopyConfiguration(const TSharedPtr<FMediaIOCoreTextureSampleBase>& SourceSample) override;
+	virtual void FreeSample() override;
+	//~ End FMediaIOCoreTextureSampleBase
+
 	/**
 	 * Initialize the sample.
 	 * 
