MediaIO: Fix incorrectly removed previous NDI defaults in CL 47872781.

#Synced-CL 49940565

diff --git a/Engine/Plugins/Media/NDIMedia/Source/NDIMedia/Private/NDIMediaSource.cpp b/Engine/Plugins/Media/NDIMedia/Source/NDIMedia/Private/NDIMediaSource.cpp
--- a/Engine/Plugins/Media/NDIMedia/Source/NDIMedia/Private/NDIMediaSource.cpp
+++ b/Engine/Plugins/Media/NDIMedia/Source/NDIMedia/Private/NDIMediaSource.cpp
@@ -22,7 +22,16 @@ UNDIMediaSource::UNDIMediaSource()
 	, MaxNumVideoFrameBuffer(8)
 	, bLogDropFrame(true)
 	, bEncodeTimecodeInTexel(false)
-{	
+{
+#if WITH_EDITORONLY_DATA
+PRAGMA_DISABLE_DEPRECATION_WARNINGS
+	bOverrideSourceEncoding_DEPRECATED = false;
+	OverrideSourceEncoding_DEPRECATED = EMediaIOCoreSourceEncoding::sRGB;
+	bOverrideSourceColorSpace_DEPRECATED = false;
+	OverrideSourceColorSpace_DEPRECATED = ETextureColorSpace::TCS_sRGB;
+PRAGMA_ENABLE_DEPRECATION_WARNINGS
+#endif
+
 	AssignDefaultConfiguration();
 }
 
