[NDI Media] Fix default stream color settings when there are no source overrides.

#Synced-CL 49891360

diff --git a/Engine/Plugins/Media/NDIMedia/Source/NDIMedia/Private/Player/NDIMediaStreamPlayer.cpp b/Engine/Plugins/Media/NDIMedia/Source/NDIMedia/Private/Player/NDIMediaStreamPlayer.cpp
--- a/Engine/Plugins/Media/NDIMedia/Source/NDIMedia/Private/Player/NDIMediaStreamPlayer.cpp
+++ b/Engine/Plugins/Media/NDIMedia/Source/NDIMedia/Private/Player/NDIMediaStreamPlayer.cpp
@@ -45,6 +45,7 @@ class FNDIMediaBinarySamplePool : public TMediaObjectPool<FMediaIOCoreBinarySamp
 FNDIMediaStreamPlayer::FNDIMediaStreamPlayer(IMediaEventSink& InEventSink)
 	: Super(InEventSink)
 	, NDIPlayerState(EMediaState::Closed)
+	, NdiColorSettings(MakeShared<FNativeMediaSourceColorSettings>())
 	, EventSink(InEventSink)
 	, TextureSamplePool(new FNDIMediaTextureSamplePool)
 	, AudioSamplePool(new FNDIMediaAudioSamplePool)
@@ -340,6 +341,8 @@ void FNDIMediaStreamPlayer::HandleVideoFrameReceived(FNDIStreamReceiver* InRecei
 		return;
 	}
 
+	UpdateNdiColorSettings();
+	
 	TSharedRef<FNDIMediaTextureSample> TextureSample = TextureSamplePool->AcquireShared();
 
 	FTimecode SourceTimecode = InReceiver->GetCurrentTimecode();
@@ -351,7 +354,7 @@ void FNDIMediaStreamPlayer::HandleVideoFrameReceived(FNDIStreamReceiver* InRecei
 		SampleTime = SourceTimecode.ToTimespan(InReceiver->GetCurrentFrameRate());
 	}
 
-	if (TextureSample->Initialize(InVideoFrame, SourceColorSettings, SampleTime, SourceTimecode))
+	if (TextureSample->Initialize(InVideoFrame, NdiColorSettings, SampleTime, SourceTimecode))
 	{
 		if (TextureSample->CustomConverter)
 		{
@@ -446,6 +449,24 @@ void FNDIMediaStreamPlayer::VerifyFrameDropCount()
 	// todo
 }
 
+void FNDIMediaStreamPlayer::UpdateNdiColorSettings()
+{
+	if (!NdiColorSettings)
+	{
+		NdiColorSettings = MakeShared<FNativeMediaSourceColorSettings>();
+	}
+
+	// Start off with the default Ndi color settings. Stream derived settings not implemented yet (todo).
+	NdiColorSettings->SetEncodingOverride(UE::Color::EEncoding::sRGB);
+	NdiColorSettings->SetColorSpaceOverride(UE::Color::FColorSpace::GetSRGB());
+
+	// Apply source overrides.
+	if (SourceColorSettings)
+	{
+		NdiColorSettings->Update(*SourceColorSettings);
+	}
+}
+
 bool FNDIMediaStreamPlayer::IsHardwareReady() const
 {
 	return NDIPlayerState == EMediaState::Playing;
diff --git a/Engine/Plugins/Media/NDIMedia/Source/NDIMedia/Private/Player/NDIMediaStreamPlayer.h b/Engine/Plugins/Media/NDIMedia/Source/NDIMedia/Private/Player/NDIMediaStreamPlayer.h
--- a/Engine/Plugins/Media/NDIMedia/Source/NDIMedia/Private/Player/NDIMediaStreamPlayer.h
+++ b/Engine/Plugins/Media/NDIMedia/Source/NDIMedia/Private/Player/NDIMediaStreamPlayer.h
@@ -67,6 +67,9 @@ private:
 
 	void VerifyFrameDropCount();
 
+	/** Update the color settings resulting from Ndi setting + media source overrides. */
+	void UpdateNdiColorSettings();
+
 private:
 	/** Keep track of the options (media source) object this player was opened with to filter global events. */
 	FObjectKey OptionsObject;
@@ -102,6 +105,9 @@ private:
 
 	/** Sample rate from the last audio frame received. */
 	uint32 NDIThreadAudioSampleRate = 0;
+
+	/** Container for local stream color settings. */
+	TSharedPtr<FNativeMediaSourceColorSettings, ESPMode::ThreadSafe> NdiColorSettings;
 	
 	/** The media event sink. */
 	IMediaEventSink& EventSink;
diff --git a/Engine/Source/Runtime/MediaAssets/Private/BaseMediaSourceColorSettings.cpp b/Engine/Source/Runtime/MediaAssets/Private/BaseMediaSourceColorSettings.cpp
--- a/Engine/Source/Runtime/MediaAssets/Private/BaseMediaSourceColorSettings.cpp
+++ b/Engine/Source/Runtime/MediaAssets/Private/BaseMediaSourceColorSettings.cpp
@@ -72,6 +72,30 @@ void FNativeMediaSourceColorSettings::Update(const FMediaSourceColorSettings& In
 	}
 }
 
+void FNativeMediaSourceColorSettings::Update(const FNativeMediaSourceColorSettings& InSettings)
+{
+	const UE::Color::EEncoding SourceEncoding = InSettings.GetEncodingOverride();
+	if (SourceEncoding != UE::Color::EEncoding::None)
+	{
+		SetEncodingOverride(SourceEncoding);
+	}
+
+	{
+		FScopeLock Lock(&ColorSpaceCriticalSection);
+		FScopeLock LockOther(&InSettings.ColorSpaceCriticalSection);
+		if (InSettings.ColorSpaceOverride.IsSet())
+		{
+			ColorSpaceOverride = InSettings.ColorSpaceOverride;
+		}
+	}
+
+	const UE::Color::EChromaticAdaptationMethod OtherChromaticAdaptationMethod = InSettings.GetChromaticAdaptationMethod();
+	if (OtherChromaticAdaptationMethod != UE::Color::EChromaticAdaptationMethod::None)
+	{
+		ChromaticAdaptationMethod = OtherChromaticAdaptationMethod;
+	}
+}
+
 bool FNativeMediaSourceColorSettings::HasColorSpaceOverride() const
 {
 	FScopeLock Lock(&ColorSpaceCriticalSection);
diff --git a/Engine/Source/Runtime/MediaAssets/Public/BaseMediaSourceColorSettings.h b/Engine/Source/Runtime/MediaAssets/Public/BaseMediaSourceColorSettings.h
--- a/Engine/Source/Runtime/MediaAssets/Public/BaseMediaSourceColorSettings.h
+++ b/Engine/Source/Runtime/MediaAssets/Public/BaseMediaSourceColorSettings.h
@@ -90,6 +90,9 @@ struct FNativeMediaSourceColorSettings : public IMediaOptions::FDataContainer
 	/** Updates the native settings from user-controlled settings. */
 	MEDIAASSETS_API void Update(const FMediaSourceColorSettings& InSettings);
 
+	/** Updates the native settings from the overrides of another native settings. */
+	MEDIAASSETS_API void Update(const FNativeMediaSourceColorSettings& InSettings);
+
 	/** Has encoding override. */
 	inline bool HasEncodingOverride() const { return EncodingOverride != UE::Color::EEncoding::None; }
 
