[NDIMedia] Allow NDI media sources to be opened even though the configured source is not yet available on the network.

There may be a delay during the startup sequence where the NDI finder doesn't detect sources. We allow the media plate to open the NDI media source any way. The NDI receiver handles this internally and will remain in "streaming" mode (not ready) until the source becomes available.

#Synced-CL 49891360

diff --git a/Engine/Plugins/Media/NDIMedia/Source/NDIMedia/Private/NDIMediaSource.cpp b/Engine/Plugins/Media/NDIMedia/Source/NDIMedia/Private/NDIMediaSource.cpp
--- a/Engine/Plugins/Media/NDIMedia/Source/NDIMedia/Private/NDIMediaSource.cpp
+++ b/Engine/Plugins/Media/NDIMedia/Source/NDIMedia/Private/NDIMediaSource.cpp
@@ -32,7 +32,11 @@ PRAGMA_DISABLE_DEPRECATION_WARNINGS
 PRAGMA_ENABLE_DEPRECATION_WARNINGS
 #endif
 
-	AssignDefaultConfiguration();
+	// Don't instantiate the source finder for the CDO.
+	if (!HasAnyFlags(RF_ClassDefaultObject))
+	{
+		AssignDefaultConfiguration();
+	}
 }
 
 bool UNDIMediaSource::GetMediaOption(const FName& Key, bool DefaultValue) const
@@ -198,11 +202,11 @@ bool UNDIMediaSource::Validate() const
 			break;
 		}
 	}
-		
+
+	// Even if the stream is currently offline, we allow the media source to be opened in case the stream become online at a later time.
 	if (!bDeviceFound)
 	{
-		UE_LOG(LogNDIMedia, Warning, TEXT("The MediaSource '%s' use the device '%s' that doesn't exist on this machine."), *GetName(), *MediaConfiguration.MediaConnection.Device.DeviceName.ToString());
-		return false;
+		UE_LOG(LogNDIMedia, Warning, TEXT("The MediaSource '%s' uses the stream '%s' that is not currently available."), *GetName(), *MediaConfiguration.MediaConnection.Device.DeviceName.ToString());
 	}
 	
 	return true;
